Fix wrong case in build artifact output path of Network Profiler.

This is a subset of changes (compiler warning fixes needed for Linux editor) from pull request 185 by sbc100, 3dluvr, wshearn, abarbu et al (#ue4linux community): https://github.com/EpicGames/UnrealEngine/pull/185

diff --git a/Engine/Source/Programs/NetworkProfiler/NetworkProfiler/NetworkProfiler.csproj b/Engine/Source/Programs/NetworkProfiler/NetworkProfiler/NetworkProfiler.csproj
--- a/Engine/Source/Programs/NetworkProfiler/NetworkProfiler/NetworkProfiler.csproj
+++ b/Engine/Source/Programs/NetworkProfiler/NetworkProfiler/NetworkProfiler.csproj
@@ -41,7 +41,7 @@
     <DebugSymbols>true</DebugSymbols>
     <DebugType>full</DebugType>
     <Optimize>false</Optimize>
-    <OutputPath>..\..\..\..\Binaries\DotNet\</OutputPath>
+    <OutputPath>..\..\..\..\Binaries\DotNET\</OutputPath>
     <DefineConstants>DEBUG;TRACE</DefineConstants>
     <ErrorReport>prompt</ErrorReport>
     <WarningLevel>4</WarningLevel>
@@ -49,7 +49,7 @@
   <PropertyGroup Condition=" '$(Configuration)|$(Platform)' == 'Development|AnyCPU' ">
     <DebugType>pdbonly</DebugType>
     <Optimize>true</Optimize>
-    <OutputPath>..\..\..\..\Binaries\DotNet\</OutputPath>
+    <OutputPath>..\..\..\..\Binaries\DotNET\</OutputPath>
     <DefineConstants>TRACE</DefineConstants>
     <ErrorReport>prompt</ErrorReport>
     <WarningLevel>4</WarningLevel>
@@ -161,4 +161,4 @@
   <Target Name="AfterBuild">
   </Target>
   -->
-</Project>
\ No newline at end of file
+</Project>
