Change overhead to UDP size

diff --git a/Engine/Source/Programs/NetworkProfiler/NetworkProfiler/NetworkStream.cs b/Engine/Source/Programs/NetworkProfiler/NetworkProfiler/NetworkStream.cs
--- a/Engine/Source/Programs/NetworkProfiler/NetworkProfiler/NetworkStream.cs
+++ b/Engine/Source/Programs/NetworkProfiler/NetworkProfiler/NetworkStream.cs
@@ -16,7 +16,8 @@ namespace NetworkProfiler
 	class NetworkStream
 	{
 		/** Per packet overhead to take into account for total outgoing bandwidth. */
-		public static int PacketOverhead = 48;
+		//public static int PacketOverhead = 48;
+		public static int PacketOverhead = 28;
 
 		/** Array of unique names. Code has fixed indexes into it.					*/
 		public List<string> NameArray = new List<string>();
