Remove unused network profiler stats

Bump the network profiler version, and add support for version checking

diff --git a/Engine/Source/Programs/NetworkProfiler/NetworkProfiler/MainWindow.cs b/Engine/Source/Programs/NetworkProfiler/NetworkProfiler/MainWindow.cs
--- a/Engine/Source/Programs/NetworkProfiler/NetworkProfiler/MainWindow.cs
+++ b/Engine/Source/Programs/NetworkProfiler/NetworkProfiler/MainWindow.cs
@@ -144,8 +144,13 @@ namespace NetworkProfiler
 			&&	CurrentFrame < CurrentNetworkStream.Frames.Count )
 			{
 				SummaryTextBox.Lines = CurrentNetworkStream.Frames[CurrentFrame].Filter(CurrentActorFilter,CurrentPropertyFilter,CurrentRPCFilter).ToStringArray();
+#if true
 				DetailTextBox.Lines = CurrentNetworkStream.Frames[CurrentFrame].Filter(CurrentActorFilter,CurrentPropertyFilter,CurrentRPCFilter).ToDetailedStringArray(CurrentActorFilter,CurrentPropertyFilter,CurrentRPCFilter);
-                
+#else
+				// Show a performance info
+				DetailTextBox.Lines = CurrentNetworkStream.Frames[CurrentFrame].Filter( CurrentActorFilter, CurrentPropertyFilter, CurrentRPCFilter ).ToActorPerformanceString( CurrentNetworkStream, CurrentActorFilter, CurrentPropertyFilter, CurrentRPCFilter );
+#endif
+           
 				// Populate the new actor performance view
 				CurrentNetworkStream.Frames[CurrentFrame].Filter(CurrentActorFilter, CurrentPropertyFilter, CurrentRPCFilter).ToActorPerformanceView(CurrentNetworkStream, ActorPerformanceView, CurrentActorFilter, CurrentPropertyFilter, CurrentRPCFilter);
             }
diff --git a/Engine/Source/Programs/NetworkProfiler/NetworkProfiler/PartialNetworkStream.cs b/Engine/Source/Programs/NetworkProfiler/NetworkProfiler/PartialNetworkStream.cs
--- a/Engine/Source/Programs/NetworkProfiler/NetworkProfiler/PartialNetworkStream.cs
+++ b/Engine/Source/Programs/NetworkProfiler/NetworkProfiler/PartialNetworkStream.cs
@@ -278,35 +278,92 @@ namespace NetworkProfiler
 
 			foreach (var UniqueActor in ActorDetailList)
 			{
-				long NumPotentialActorBytes = ( UniqueActor.Value.PotentialSizeBits + 7 ) / 8;
 				long NumActorBytes = ( UniqueActor.Value.SizeBits + 7 ) / 8;
-				string ActorStr = string.Format( "{0,-32} : {1:0.00} ({2:00}/{3:00}) ({4:00})", NetworkStream.GetName( UniqueActor.Key ), UniqueActor.Value.TimeInMS, NumPotentialActorBytes, NumActorBytes, UniqueActor.Value.Count );
+				string ActorStr = string.Format( "{0,-32} : {1:0.00} ({2:000}) ({3:000})", NetworkStream.GetName( UniqueActor.Key ), UniqueActor.Value.TimeInMS, NumActorBytes, UniqueActor.Value.Count );
 				TreeView.Nodes.Add(ActorStr);
 
-				var PropertyDetailList = UniqueActor.Value.Properties.UniqueItems.OrderByDescending(s => s.Value.TimeInMS).ToList();
+				var PropertyDetailList = UniqueActor.Value.Properties.UniqueItems.OrderByDescending(s => s.Value.SizeBits).ToList();
 
 				foreach (var Property in PropertyDetailList)
 				{
-					long NumPotentialPropBytes = ( Property.Value.PotentialSizeBits + 7 ) / 8;
-					long NumPropBytes = ( Property.Value.SizeBits + 7 ) / 8;
+					long NumPropBytes	= ( Property.Value.SizeBits + 7 ) / 8;
+					string PropName		= NetworkStream.GetName( Property.Key );
 
-					string PropName = NetworkStream.GetName( Property.Key );
-					String FlagsStr = Property.Value.FirstItem.GenerateFlagsString();
-
-					string PropStr = string.Format( "{0,-22} {1,-8} : {2:0.00} ({3:00}/{4:00}) ({5:00})", PropName, FlagsStr, Property.Value.TimeInMS, NumPotentialPropBytes, NumPropBytes, Property.Value.Count );
+					string PropStr = string.Format( "{0,-25} : {1:000} ({2:000})", PropName, NumPropBytes, Property.Value.Count );
 					TreeView.Nodes[TreeView.Nodes.Count - 1].Nodes.Add(PropStr);
 				}
 			}
         }
-        
-        /**
-         * Converts the passed in number of bytes to a string formatted as Bytes, KByte, MByte depending
-         * on magnitude.
-         * 
-         * @param	SizeInBytes		Size in bytes to conver to formatted string
-         * 
-         * @return string representation of value, either Bytes, KByte or MByte
-         */ 
+
+		public string[] ToActorPerformanceString( NetworkStream NetworkStream, string ActorFilter, string PropertyFilter, string RPCFilter )
+		{
+			var Details = new List<string>();
+
+ 			UniqueItemTracker<UniqueActor, TokenReplicateActor> UniqueActors = new UniqueItemTracker<UniqueActor, TokenReplicateActor>();
+			foreach ( TokenBase Token in Tokens )
+			{
+				TokenReplicateActor ActorToken = Token as TokenReplicateActor;
+
+				if ( ActorToken == null )
+				{
+					continue;
+				}
+
+				UniqueActors.AddItem( ActorToken, NetworkStream.GetClassNameIndex( ActorToken.ActorNameIndex ) );
+			}
+
+			var ActorDetailList = UniqueActors.UniqueItems.OrderByDescending( s => s.Value.TimeInMS ).ToList();
+
+			Int32 NumActors = 0;
+			long NumSentBytes = 0;
+			float TotalMS = 0.0f;
+
+			foreach ( var UniqueActor in ActorDetailList )
+			{
+				NumActors += UniqueActor.Value.Count;
+				TotalMS += UniqueActor.Value.TimeInMS;
+				var PropertyDetailList = UniqueActor.Value.Properties.UniqueItems.OrderByDescending( s => s.Value.SizeBits ).ToList();
+
+				foreach ( var Property in PropertyDetailList )
+				{
+					NumSentBytes += ( Property.Value.SizeBits + 7 ) / 8;
+				}
+			}
+
+			Details.Add( "Total Actors  : " + NumActors.ToString() );
+			Details.Add( "Total MS      : " + string.Format( "{0:0.00}", TotalMS ) );
+			Details.Add( "Sent Bytes    : " + NumSentBytes.ToString() );
+
+			foreach ( var UniqueActor in ActorDetailList )
+			{
+				long NumActorBytes = ( UniqueActor.Value.SizeBits + 7 ) / 8;
+				string ActorStr = string.Format( "{0,-34} : {1:0.00} ({2,4}) ({3,2})", NetworkStream.GetName( UniqueActor.Key ), UniqueActor.Value.TimeInMS, NumActorBytes, UniqueActor.Value.Count );
+				Details.Add( ActorStr );
+
+				var PropertyDetailList = UniqueActor.Value.Properties.UniqueItems.OrderByDescending( s => s.Value.SizeBits ).ToList();
+
+				foreach ( var Property in PropertyDetailList )
+				{
+					long NumPropBytes	= ( Property.Value.SizeBits + 7 ) / 8;
+					string PropName		= NetworkStream.GetName( Property.Key );
+
+					string PropStr = string.Format( "{0,-22} : {1,4} ({2,2})", PropName, NumPropBytes, Property.Value.Count );
+					Details.Add( "   " + PropStr );
+				}
+				Details.Add( "-----------------------------------" );
+			}
+
+			return Details.ToArray();
+		}
+
+		/**
+		 * Converts the passed in number of bytes to a string formatted as Bytes, KByte, MByte depending
+		 * on magnitude.
+		 * 
+		 * @param	SizeInBytes		Size in bytes to conver to formatted string
+		 * 
+		 * @return string representation of value, either Bytes, KByte or MByte
+		 */ 
 		public string ConvertToSizeString( float SizeInBytes )
 		{
 			// Format as MByte if size > 1 MByte.
@@ -496,15 +553,11 @@ namespace NetworkProfiler
 	*/
 	class UniqueProperty : UniqueItem<TokenReplicateProperty>
 	{
-		public float	TimeInMS = 0;
-		public long		PotentialSizeBits = 0;
 		public long		SizeBits = 0;
 
 		override public void OnItemAdded(TokenReplicateProperty Item)
 		{
 			base.OnItemAdded(Item);
-			TimeInMS += Item.TimeInMS;
-			PotentialSizeBits += Item.NumPotentialBits;
 			SizeBits += Item.NumBits;
 		}
 	}
@@ -515,7 +568,6 @@ namespace NetworkProfiler
 	*/
 	class UniqueActor : UniqueItem<TokenReplicateActor>
     {
-		public long		PotentialSizeBits = 0;
 		public long		SizeBits = 0;
         public float	TimeInMS = 0;
 
@@ -529,7 +581,6 @@ namespace NetworkProfiler
 			foreach (var Property in Item.Properties)
 			{
 				SizeBits += Property.NumBits;
-				PotentialSizeBits += Property.NumPotentialBits;
 				Properties.AddItem(Property, Property.PropertyNameIndex);
 			}
 		}
diff --git a/Engine/Source/Programs/NetworkProfiler/NetworkProfiler/StreamHeader.cs b/Engine/Source/Programs/NetworkProfiler/NetworkProfiler/StreamHeader.cs
--- a/Engine/Source/Programs/NetworkProfiler/NetworkProfiler/StreamHeader.cs
+++ b/Engine/Source/Programs/NetworkProfiler/NetworkProfiler/StreamHeader.cs
@@ -3,6 +3,8 @@
 
 using System;
 using System.IO;
+using System.Drawing;
+using System.Windows.Forms;
 
 namespace NetworkProfiler
 {
@@ -12,7 +14,10 @@ namespace NetworkProfiler
 	public class StreamHeader
 	{
 		/** Magic number at beginning of data file.					*/
-		public const UInt32 ExpectedMagic = 0x1DBF348A;
+		public const UInt32 ExpectedMagic = 0x1DBF348C;
+
+		/** We expect this version, or we can't proceed.			*/
+		public const UInt32 ExpectedVersion = 5;
 
 		/** Magic to ensure we're opening the right file.			*/
 		public UInt32 Magic;
@@ -31,6 +36,37 @@ namespace NetworkProfiler
 		/** URL used to open/ browse to the map.					*/
 		public string URL;
 
+		public static void ShowDialog( string Header, string Message )
+		{
+			Form Dialog = new Form();
+
+			Dialog.Width			= 400;
+			Dialog.Height			= 120;
+			Dialog.Text				= Header;
+			Dialog.FormBorderStyle	= FormBorderStyle.FixedDialog;
+
+			Label LabelMessage = new Label() 
+			{ 
+				//Left		= 0, 
+				//Top		= 0, 
+				Width		= Dialog.Width,
+				Height		= Dialog.Height / 2,
+				Text		= Message, 
+				TextAlign	= ContentAlignment.MiddleCenter,
+				Font		= new Font( SystemFonts.MessageBoxFont.FontFamily.Name, 12.0f, FontStyle.Bold, GraphicsUnit.Point ) 
+			};
+
+			LabelMessage.TextAlign = ContentAlignment.MiddleCenter;
+
+			Button ButtonOK = new Button() { Text = "OK", Left = Dialog.Width - 100, Width = 80, Top = Dialog.Height - 70 };
+
+			ButtonOK.Click += ( sender, e ) => { Dialog.Close(); };
+
+			Dialog.Controls.Add( ButtonOK );
+			Dialog.Controls.Add( LabelMessage );
+			Dialog.ShowDialog();
+		}
+
 		/**
 		 * Reads the header information from the passed in stream and returns it. It also returns
 		 * a BinaryReader that is endian-appropriate for the data stream. 
@@ -63,6 +99,13 @@ namespace NetworkProfiler
 			// At this point we should have a valid header. If no, throw an exception.
 			if( Header.Magic != StreamHeader.ExpectedMagic )
 			{
+				ShowDialog( "Error", "Unrecognized file" );
+				throw new InvalidDataException();
+			}
+
+			if ( Header.Version != StreamHeader.ExpectedVersion )
+			{
+				ShowDialog( "Error", "Invalid version" );
 				throw new InvalidDataException();
 			}
 
diff --git a/Engine/Source/Programs/NetworkProfiler/NetworkProfiler/StreamParser.cs b/Engine/Source/Programs/NetworkProfiler/NetworkProfiler/StreamParser.cs
--- a/Engine/Source/Programs/NetworkProfiler/NetworkProfiler/StreamParser.cs
+++ b/Engine/Source/Programs/NetworkProfiler/NetworkProfiler/StreamParser.cs
@@ -159,7 +159,7 @@ namespace NetworkProfiler
 					if( Token.TokenType == ETokenTypes.ReplicateProperty )
 					{
 						var TokenReplicateProperty = Token as TokenReplicateProperty;
-						NetworkStream.UpdateSummary(ref NetworkStream.PropertyNameToSummary, TokenReplicateProperty.PropertyNameIndex, TokenReplicateProperty.NumBits, TokenReplicateProperty.TimeInMS );
+						NetworkStream.UpdateSummary(ref NetworkStream.PropertyNameToSummary, TokenReplicateProperty.PropertyNameIndex, TokenReplicateProperty.NumBits, 0 );
 						//LastActorToken.Properties.Add(TokenReplicateProperty);
                         LastProperties.Add(TokenReplicateProperty);
 					}					
diff --git a/Engine/Source/Programs/NetworkProfiler/NetworkProfiler/Tokens.cs b/Engine/Source/Programs/NetworkProfiler/NetworkProfiler/Tokens.cs
--- a/Engine/Source/Programs/NetworkProfiler/NetworkProfiler/Tokens.cs
+++ b/Engine/Source/Programs/NetworkProfiler/NetworkProfiler/Tokens.cs
@@ -393,64 +393,16 @@ namespace NetworkProfiler
 	{
 		/** Name table index of property name. */
 		public int		PropertyNameIndex;
-		/** Number of bits that could be sent (but may not because they didn't change). */
-		public UInt16	NumPotentialBits;
 		/** Number of bits serialized/ sent. */
 		public UInt16	NumBits;
-		/** True if property is dynamic (struct, TArray, etc) */
-		public bool		bIsDynamic;
-		/** True if property is from a component */
-		public bool		bIsComponent;
-		/** Time in ms for this property to replicate */
-		public float	TimeInMS;
 
 		/** Constructor, serializing members from passed in stream. */
 		public TokenReplicateProperty(BinaryReader BinaryStream)
 		{
-            byte Flags = BinaryStream.ReadByte();
-            bIsDynamic = (Flags & 1) != 0;
-            bIsComponent = (Flags & 2) != 0;
-
-            TimeInMS = BinaryStream.ReadSingle();
-
             PropertyNameIndex = BinaryStream.ReadInt32();
-			NumPotentialBits = BinaryStream.ReadUInt16();
 			NumBits = BinaryStream.ReadUInt16();
 		}
 
-		/** Helper to build property flag string */
-		private void ConcatPropertyType( bool value, string name, ref string CurrentString )
-		{
-			if (!value)
-			{
-				return;
-			}
-			if (CurrentString == "")
-			{
-				CurrentString = " (";
-			}
-			else
-			{
-				CurrentString += ", ";
-			}
-			CurrentString += name;
-		}
-
-		/** Helper to build property flag string */
-		public String GenerateFlagsString()
-		{
-			String FlagsStr = "";
-
-			ConcatPropertyType(bIsDynamic, "D", ref FlagsStr);
-			ConcatPropertyType(bIsComponent, "C", ref FlagsStr);
-			if (FlagsStr != "")
-			{
-				FlagsStr += ")";
-			}
-
-			return FlagsStr;
-		}
-
 		/**
 		 * Returns whether the token matches/ passes based on the passed in filters.
 		 * 
diff --git a/Engine/Source/Runtime/Engine/Private/NetworkProfiler.cpp b/Engine/Source/Runtime/Engine/Private/NetworkProfiler.cpp
--- a/Engine/Source/Runtime/Engine/Private/NetworkProfiler.cpp
+++ b/Engine/Source/Runtime/Engine/Private/NetworkProfiler.cpp
@@ -24,9 +24,9 @@
 FNetworkProfiler GNetworkProfiler;
 
 /** Magic value, determining that file is a network profiler file.				*/
-#define NETWORK_PROFILER_MAGIC						0x1DBF348A
+#define NETWORK_PROFILER_MAGIC						0x1DBF348C
 /** Version of memory profiler. Incremented on serialization changes.			*/
-#define NETWORK_PROFILER_VERSION					4
+#define NETWORK_PROFILER_VERSION					5
 
 enum ENetworkProfilingPayloadType
 {
@@ -308,20 +308,15 @@ void FNetworkProfiler::TrackReplicateActor( const AActor* Actor, FReplicationFla
  * @param	Property	Property being replicated
  * @param	NumBits		Number of bits used to replicate this property
  */
-void FNetworkProfiler::TrackReplicateProperty( const UProperty* Property, bool bIsDynamicProperty, bool bIsComponentProperty, uint32 Cycles, uint16 NumPotentialBits, uint16 NumBits )
+void FNetworkProfiler::TrackReplicateProperty( const UProperty* Property, uint16 NumBits )
 {
 	if( bIsTrackingEnabled )
 	{
 		SCOPE_LOCK_REF(CriticalSection);
 		uint8 Type = NPTYPE_ReplicateProperty;
 		(*FileWriter) << Type;
-		uint8 Flags = (bIsDynamicProperty << 0) | (bIsComponentProperty << 1);
-		(*FileWriter) << Flags;
-		float TimeInMS = FPlatformTime::ToMilliseconds(Cycles);	// FIXME: We may want to just pass in cycles to profiler to we don't lose precision
-		(*FileWriter) << TimeInMS;
 		int32 NameTableIndex = GetNameTableIndex( Property->GetName() );
 		(*FileWriter) << NameTableIndex;
-		(*FileWriter) << NumPotentialBits;
 		(*FileWriter) << NumBits;
 	}
 }
diff --git a/Engine/Source/Runtime/Engine/Private/RepLayout.cpp b/Engine/Source/Runtime/Engine/Private/RepLayout.cpp
--- a/Engine/Source/Runtime/Engine/Private/RepLayout.cpp
+++ b/Engine/Source/Runtime/Engine/Private/RepLayout.cpp
@@ -302,10 +302,6 @@ bool FRepLayout::CompareProperties(
 
 	for ( const uint16 * RESTRICT pLifeProp = FirstProp; pLifeProp < LastProp; ++pLifeProp )
 	{
-//#if USE_NETWORK_PROFILER 
-		//const uint32 PropertyStartTime = GNetworkProfiler.IsTrackingEnabled() ? FPlatformTime::Cycles() : 0;
-//#endif
-
 		const FRepParentCmd & ParentCmd = Parents[*pLifeProp];
 
 		// We store changed properties on each parent, so we can build a final sorted change list later
@@ -340,8 +336,6 @@ bool FRepLayout::CompareProperties(
 			// Something changed on this parent property
 			PropertyChanged = true;
 		}
-
-		//NETWORK_PROFILER( GNetworkProfiler.TrackReplicateProperty( ParentCmd.Property, true, false, FPlatformTime::Cycles() - PropertyStartTime, ParentCmd.Property->ElementSize * 8, 0 ) );
 	}
 
 	return PropertyChanged;
@@ -909,7 +903,7 @@ uint16 FRepLayout::SendProperties_r(
 
 			const FRepParentCmd & ParentCmd = Parents[Cmd.ParentIndex];
 
-			NETWORK_PROFILER( GNetworkProfiler.TrackReplicateProperty( ParentCmd.Property, false, false, 0, 0, NumEndBits - NumStartBits ) );
+			NETWORK_PROFILER( GNetworkProfiler.TrackReplicateProperty( ParentCmd.Property, NumEndBits - NumStartBits ) );
 
 			// Make the shadow state match the actual state at the time of send
 			StoreProperty( Cmd, (void*)( StoredData + Cmd.Offset ), (const void*)( Data + Cmd.Offset ) );
diff --git a/Engine/Source/Runtime/Engine/Public/Net/NetworkProfiler.h b/Engine/Source/Runtime/Engine/Public/Net/NetworkProfiler.h
--- a/Engine/Source/Runtime/Engine/Public/Net/NetworkProfiler.h
+++ b/Engine/Source/Runtime/Engine/Public/Net/NetworkProfiler.h
@@ -129,7 +129,7 @@ public:
 	 * @param	Property	Property being replicated
 	 * @param	NumBits		Number of bits used to replicate this property
 	 */
-	void TrackReplicateProperty( const UProperty* Property, bool bIsDynamicProperty, bool bIsComponentProperty, uint32 Cycles, uint16 NumPotentialBits, uint16 NumBits );
+	void TrackReplicateProperty( const UProperty* Property, uint16 NumBits );
 
 	/**
 	 * Track event occuring, like e.g. client join/ leave
