Make network profiler scroll bar visible

diff --git a/Engine/Source/Programs/NetworkProfiler/NetworkProfiler/Charting.cs b/Engine/Source/Programs/NetworkProfiler/NetworkProfiler/Charting.cs
--- a/Engine/Source/Programs/NetworkProfiler/NetworkProfiler/Charting.cs
+++ b/Engine/Source/Programs/NetworkProfiler/NetworkProfiler/Charting.cs
@@ -8,6 +8,7 @@ using System.Linq;
 using System.Text;
 using System.Windows.Forms.DataVisualization;
 using System.Windows.Forms.DataVisualization.Charting;
+using System.Drawing;
 
 // Microsoft Chart Controls Add-on for Microsoft Visual Studio 2008
 //
@@ -31,6 +32,16 @@ namespace NetworkProfiler
 			NetworkChart.ResetAutoValues();
 			NetworkChart.Invalidate();
 
+			NetworkChart.ChartAreas[0].AxisX.ScrollBar.IsPositionedInside = false;
+			NetworkChart.ChartAreas[0].AxisX.ScrollBar.ButtonStyle = ScrollBarButtonStyles.All;
+			NetworkChart.ChartAreas[0].AxisX.ScrollBar.Size = 15;
+			NetworkChart.ChartAreas[0].AxisX.ScrollBar.ButtonColor = Color.LightGray;
+
+			NetworkChart.ChartAreas[0].AxisY.ScrollBar.IsPositionedInside = false;
+			NetworkChart.ChartAreas[0].AxisY.ScrollBar.ButtonStyle = ScrollBarButtonStyles.All;
+			NetworkChart.ChartAreas[0].AxisY.ScrollBar.Size = 15;
+			NetworkChart.ChartAreas[0].AxisY.ScrollBar.ButtonColor = Color.LightGray;
+
 			int FrameCounter = 0;
 			foreach( PartialNetworkStream RawFrame in NetworkStream.Frames )
 			{
@@ -42,8 +53,8 @@ namespace NetworkProfiler
 				}
 				float OneOverDeltaTime = 1 / (Frame.EndTime - Frame.StartTime);
 
-                NetworkChart.Series["ActorCount"].Points.AddXY( FrameCounter, Frame.ActorCount );
-                NetworkChart.Series["ActorCountSec"].Points.AddXY(FrameCounter, Frame.ActorCount * OneOverDeltaTime);
+				NetworkChart.Series["ActorCount"].Points.AddXY( FrameCounter, Frame.ActorCount );
+				NetworkChart.Series["ActorCountSec"].Points.AddXY( FrameCounter, Frame.ActorCount * OneOverDeltaTime );
 				NetworkChart.Series["PropertyCount"].Points.AddXY( FrameCounter, Frame.PropertyCount );
 				NetworkChart.Series["PropertyCountSec"].Points.AddXY( FrameCounter, Frame.PropertyCount * OneOverDeltaTime );
 				NetworkChart.Series["PropertySize"].Points.AddXY( FrameCounter, Frame.ReplicatedSizeBits / 8 );
@@ -67,13 +78,13 @@ namespace NetworkProfiler
 				int OutgoingBandwidth = Frame.UnrealSocketSize + Frame.OtherSocketSize + NetworkStream.PacketOverhead * (Frame.UnrealSocketCount + Frame.OtherSocketCount);
 				NetworkChart.Series["OutgoingBandwidthSize"].Points.AddXY( FrameCounter, OutgoingBandwidth );
 				NetworkChart.Series["OutgoingBandwidthSizeSec"].Points.AddXY( FrameCounter, OutgoingBandwidth * OneOverDeltaTime );
-                NetworkChart.Series["ActorReplicateTimeInMS"].Points.AddXY(FrameCounter, Frame.ActorReplicateTimeInMS);
-				
+				NetworkChart.Series["ActorReplicateTimeInMS"].Points.AddXY(FrameCounter, Frame.ActorReplicateTimeInMS);
+
 				if( Frame.NumEvents > 0 )
 				{
 					NetworkChart.Series["Events"].Points.AddXY( FrameCounter, 0 );
 				}
-				
+
 				FrameCounter++;
 			}
 
