Network profiler bandwidth tracking improvements:
-Breaks down bits used in tracked RPCs by header, parameters, and footer
-Breaks down bits in bunches by header and payload, and accounts for merged bunches
-Tracks sent acks
-Every bit sent from UIpConnection::LowLevelSend is now accounted for

diff --git a/Engine/Source/Programs/NetworkProfiler/NetworkProfiler/Charting.cs b/Engine/Source/Programs/NetworkProfiler/NetworkProfiler/Charting.cs
--- a/Engine/Source/Programs/NetworkProfiler/NetworkProfiler/Charting.cs
+++ b/Engine/Source/Programs/NetworkProfiler/NetworkProfiler/Charting.cs
@@ -63,6 +63,10 @@ namespace NetworkProfiler
 				NetworkChart.Series["RPCCountSec"].Points.AddXY( FrameCounter, Frame.RPCCount * OneOverDeltaTime );
 				NetworkChart.Series["RPCSize"].Points.AddXY( FrameCounter, Frame.RPCSizeBits / 8 );
 				NetworkChart.Series["RPCSizeSec"].Points.AddXY( FrameCounter, Frame.RPCSizeBits / 8 * OneOverDeltaTime );
+				NetworkChart.Series["SendAckCount"].Points.AddXY( FrameCounter, Frame.SendAckCount );
+				NetworkChart.Series["SendAckCountSec"].Points.AddXY( FrameCounter, Frame.SendAckCount * OneOverDeltaTime );
+				NetworkChart.Series["SendAckSize"].Points.AddXY( FrameCounter, Frame.SendAckSizeBits / 8 );
+				NetworkChart.Series["SendAckSizeSec"].Points.AddXY( FrameCounter, Frame.SendAckSizeBits / 8 * OneOverDeltaTime );
 				NetworkChart.Series["SendBunchCount"].Points.AddXY( FrameCounter, Frame.SendBunchCount );
 				NetworkChart.Series["SendBunchCountSec"].Points.AddXY( FrameCounter, Frame.SendBunchCount * OneOverDeltaTime );
 				NetworkChart.Series["SendBunchSize"].Points.AddXY( FrameCounter, Frame.SendBunchSizeBits / 8 );
diff --git a/Engine/Source/Programs/NetworkProfiler/NetworkProfiler/MainWindow.cs b/Engine/Source/Programs/NetworkProfiler/NetworkProfiler/MainWindow.cs
--- a/Engine/Source/Programs/NetworkProfiler/NetworkProfiler/MainWindow.cs
+++ b/Engine/Source/Programs/NetworkProfiler/NetworkProfiler/MainWindow.cs
@@ -35,12 +35,12 @@ namespace NetworkProfiler
 			{
 				Series.Enabled = false;
 			}
-			NetworkChart.Series[4].Enabled = true;
+			NetworkChart.Series["PropertySize"].Enabled = true;
 			ChartListBox.SetItemChecked(4, true);
-			NetworkChart.Series[8].Enabled = true;
+			NetworkChart.Series["RPCSize"].Enabled = true;
 			ChartListBox.SetItemChecked(8, true);
-			NetworkChart.Series[23].Enabled = true;
-			ChartListBox.SetItemChecked(23, true);
+			NetworkChart.Series["Events"].Enabled = true;
+			ChartListBox.SetItemChecked(27, true);
 
 			// Force the columns to the way we want them
 			SetupColumns( ActorListView,	new String[]	{ "Total Size (KBytes)", "Count", "Average Size (Bytes)", "Average Size (Bits)", "Time (ms)", "Actor Class" });
diff --git a/Engine/Source/Programs/NetworkProfiler/NetworkProfiler/PartialNetworkStream.cs b/Engine/Source/Programs/NetworkProfiler/NetworkProfiler/PartialNetworkStream.cs
--- a/Engine/Source/Programs/NetworkProfiler/NetworkProfiler/PartialNetworkStream.cs
+++ b/Engine/Source/Programs/NetworkProfiler/NetworkProfiler/PartialNetworkStream.cs
@@ -68,6 +68,13 @@ namespace NetworkProfiler
 		/** Total size of bytes sent on non-"Unreal" socket type. */
 		public int OtherSocketSize = 0;
 
+		// Acks
+		
+		/** Number of acks sent. */
+		public int SendAckCount = 0;
+		/** Total size of acks sent. */
+		public int SendAckSizeBits = 0;
+
 		// Detailed information.
 
 		/** List of all tokens in this substream. */
@@ -203,14 +210,14 @@ namespace NetworkProfiler
 					case ETokenTypes.SendBunch:
 						var TokenSendBunch = (TokenSendBunch) Token;
 						SendBunchCount++;
-						SendBunchSizeBits += TokenSendBunch.NumBits;
+						SendBunchSizeBits += TokenSendBunch.GetNumTotalBits();
 						SendBunchCountPerChannel[TokenSendBunch.ChannelType]++;
-						SendBunchSizeBitsPerChannel[TokenSendBunch.ChannelType] += TokenSendBunch.NumBits;
+						SendBunchSizeBitsPerChannel[TokenSendBunch.ChannelType] += TokenSendBunch.GetNumTotalBits();
 						break;
 					case ETokenTypes.SendRPC:
 						var TokenSendRPC = (TokenSendRPC) Token;
 						RPCCount++;
-						RPCSizeBits += TokenSendRPC.NumBits;
+						RPCSizeBits += TokenSendRPC.GetNumTotalBits();
 						break;
 					case ETokenTypes.ReplicateActor:
 						var TokenReplicateActor = (TokenReplicateActor) Token;
@@ -231,6 +238,11 @@ namespace NetworkProfiler
 						break;
 					case ETokenTypes.RawSocketData:
 						break;
+					case ETokenTypes.SendAck:
+						var TokenSendAck = (TokenSendAck) Token;
+						SendAckCount++;
+						SendAckSizeBits += TokenSendAck.NumBits;
+						break;
 					default:
 						throw new System.IO.InvalidDataException();
 				}
@@ -413,45 +425,49 @@ namespace NetworkProfiler
 									"Replicated Size        : {5}^"		+
 									"RPC Count              : {6}^"		+
 									"RPC Size               : {7}^"		+
-									"SendBunch Count        : {8}^"		+
-									"   Control             : {9}^"		+
-									"   Actor               : {10}^"	+
-									"   File                : {11}^"	+
-									"   Voice               : {12}^"	+
-									"SendBunch Size         : {13}^"	+
-									"   Control             : {14}^"	+
-									"   Actor               : {15}^"	+
-									"   File                : {16}^"	+
-									"   Voice               : {17}^"	+
-									"Game Socket Send Count : {18}^"	+
-									"Game Socket Send Size  : {19}^"	+
-									"Misc Socket Send Count : {20}^"	+
-									"Misc Socket Send Size  : {21}^"	+
-									"Outgoing bandwidth     : {22}^"	+
+									"Sent ack count         : {8}^"		+
+									"Sent ack size          : {9}^"		+
+									"SendBunch Count        : {10}^"	+
+									"   Control             : {11}^"	+
+									"   Actor               : {12}^"	+
+									"   File                : {13}^"	+
+									"   Voice               : {14}^"	+
+									"SendBunch Size         : {15}^"	+
+									"   Control             : {16}^"	+
+									"   Actor               : {17}^"	+
+									"   File                : {18}^"	+
+									"   Voice               : {19}^"	+
+									"Game Socket Send Count : {20}^"	+
+									"Game Socket Send Size  : {21}^"	+
+									"Misc Socket Send Count : {22}^"	+
+									"Misc Socket Send Size  : {23}^"	+
+									"Outgoing bandwidth     : {24}^"	+
 									"^"									+
 									"^"									+
 									"Network Summary per second^"		+
 									"^"									+
-									"Actor Count            : {23}^"	+
-									"Property Count         : {24}^"	+
-									"Replicated Size        : {25}^"	+
-									"RPC Count              : {26}^"	+
-									"RPC Size               : {27}^"	+
-									"SendBunch Count        : {28}^"	+
-									"   Control             : {29}^"	+
-									"   Actor               : {30}^"	+
-									"   File                : {31}^"	+
-									"   Voice               : {32}^"	+
-									"SendBunch Size         : {33}^"	+
-									"   Control             : {34}^"	+
-									"   Actor               : {35}^"	+
-									"   File                : {36}^"	+
-									"   Voice               : {37}^"	+
-									"Game Socket Send Count : {38}^"	+
-									"Game Socket Send Size  : {39}^"	+
-									"Misc Socket Send Count : {40}^"	+
-									"Misc Socket Send Size  : {41}^"	+
-									"Outgoing bandwidth     : {42}^";
+									"Actor Count            : {25}^"	+
+									"Property Count         : {26}^"	+
+									"Replicated Size        : {27}^"	+
+									"RPC Count              : {28}^"	+
+									"RPC Size               : {29}^"	+
+									"Sent ack count         : {30}^"	+
+									"Sent ack size          : {31}^"	+
+									"SendBunch Count        : {32}^"	+
+									"   Control             : {33}^"	+
+									"   Actor               : {34}^"	+
+									"   File                : {35}^"	+
+									"   Voice               : {36}^"	+
+									"SendBunch Size         : {37}^"	+
+									"   Control             : {38}^"	+
+									"   Actor               : {39}^"	+
+									"   File                : {40}^"	+
+									"   Voice               : {41}^"	+
+									"Game Socket Send Count : {42}^"	+
+									"Game Socket Send Size  : {43}^"	+
+									"Misc Socket Send Count : {44}^"	+
+									"Misc Socket Send Size  : {45}^"	+
+									"Outgoing bandwidth     : {46}^";
 
 			float OneOverDeltaTime = 1 / (EndTime - StartTime);
 
@@ -464,6 +480,8 @@ namespace NetworkProfiler
 									ConvertToSizeString(ReplicatedSizeBits / 8.0f),
 									ConvertToCountString(RPCCount),
 									ConvertToSizeString(RPCSizeBits / 8.0f),
+									ConvertToCountString(SendAckCount),
+									ConvertToSizeString(SendAckSizeBits / 8.0f),
 									ConvertToCountString(SendBunchCount),
 									ConvertToCountString(SendBunchCountPerChannel[(int)EChannelTypes.Control]),
 									ConvertToCountString(SendBunchCountPerChannel[(int)EChannelTypes.Actor]),
@@ -484,6 +502,8 @@ namespace NetworkProfiler
 									ConvertToSizeString(ReplicatedSizeBits / 8.0f * OneOverDeltaTime),
 									ConvertToCountString(RPCCount * OneOverDeltaTime),
 									ConvertToSizeString(RPCSizeBits / 8.0f * OneOverDeltaTime),
+									ConvertToCountString(SendAckCount * OneOverDeltaTime),
+									ConvertToSizeString(SendAckSizeBits / 8.0f * OneOverDeltaTime),
 									ConvertToCountString(SendBunchCount * OneOverDeltaTime),
 									ConvertToCountString(SendBunchCountPerChannel[(int)EChannelTypes.Control] * OneOverDeltaTime),
 									ConvertToCountString(SendBunchCountPerChannel[(int)EChannelTypes.Actor] * OneOverDeltaTime),
diff --git a/Engine/Source/Programs/NetworkProfiler/NetworkProfiler/StreamHeader.cs b/Engine/Source/Programs/NetworkProfiler/NetworkProfiler/StreamHeader.cs
--- a/Engine/Source/Programs/NetworkProfiler/NetworkProfiler/StreamHeader.cs
+++ b/Engine/Source/Programs/NetworkProfiler/NetworkProfiler/StreamHeader.cs
@@ -17,7 +17,7 @@ namespace NetworkProfiler
 		public const UInt32 ExpectedMagic = 0x1DBF348C;
 
 		/** We expect this version, or we can't proceed.			*/
-		public const UInt32 ExpectedVersion = 5;
+		public const UInt32 ExpectedVersion = 6;
 
 		/** Magic to ensure we're opening the right file.			*/
 		public UInt32 Magic;
diff --git a/Engine/Source/Programs/NetworkProfiler/NetworkProfiler/StreamParser.cs b/Engine/Source/Programs/NetworkProfiler/NetworkProfiler/StreamParser.cs
--- a/Engine/Source/Programs/NetworkProfiler/NetworkProfiler/StreamParser.cs
+++ b/Engine/Source/Programs/NetworkProfiler/NetworkProfiler/StreamParser.cs
@@ -152,7 +152,7 @@ namespace NetworkProfiler
 					else if( Token.TokenType == ETokenTypes.SendRPC )
 					{
 						var TokenSendRPC = Token as TokenSendRPC;
-						NetworkStream.UpdateSummary( ref NetworkStream.RPCNameToSummary, TokenSendRPC.FunctionNameIndex, TokenSendRPC.NumBits, 0.0f );
+						NetworkStream.UpdateSummary( ref NetworkStream.RPCNameToSummary, TokenSendRPC.FunctionNameIndex, TokenSendRPC.GetNumTotalBits(), 0.0f );
 					}
 
 					// Add properties to the actor token instead of network stream and keep track of summary.
diff --git a/Engine/Source/Programs/NetworkProfiler/NetworkProfiler/Tokens.cs b/Engine/Source/Programs/NetworkProfiler/NetworkProfiler/Tokens.cs
--- a/Engine/Source/Programs/NetworkProfiler/NetworkProfiler/Tokens.cs
+++ b/Engine/Source/Programs/NetworkProfiler/NetworkProfiler/Tokens.cs
@@ -21,6 +21,7 @@ namespace NetworkProfiler
 		EndOfStreamMarker,			// End of stream marker		
 		Event,						// Event
 		RawSocketData,				// Raw socket data being sent
+		SendAck,					// Ack being sent
 		MaxAndInvalid,				// Invalid token, also used as the max token index
 	}
 
@@ -91,6 +92,9 @@ namespace NetworkProfiler
 				case ETokenTypes.RawSocketData:
 					SerializedToken = new TokenRawSocketData( BinaryStream );
 					break;
+				case ETokenTypes.SendAck:
+					SerializedToken = new TokenSendAck( BinaryStream );
+					break;
 				default:
 					throw new InvalidDataException();
 			}
@@ -167,6 +171,14 @@ namespace NetworkProfiler
 		public int SocketNameIndex;
 		/** Bytes actually sent by low level code. */
 		public UInt16 BytesSent;
+		/** Number of bits representing the packet id */
+		public UInt16 NumPacketIdBits;
+		/** Number of bits representing bunches */
+		public UInt16 NumBunchBits;
+		/** Number of bits representing acks */
+		public UInt16 NumAckBits;
+		/** Number of bits used for padding */
+		public UInt16 NumPaddingBits;
 		/** IP in network byte order. */
 		public UInt32 NetworkByteOrderIP;
 
@@ -176,6 +188,10 @@ namespace NetworkProfiler
 			ThreadId = BinaryStream.ReadUInt32();
 			SocketNameIndex = BinaryStream.ReadInt32();
 			BytesSent = BinaryStream.ReadUInt16();
+			NumPacketIdBits = BinaryStream.ReadUInt16();
+			NumBunchBits = BinaryStream.ReadUInt16();
+			NumAckBits = BinaryStream.ReadUInt16();
+			NumPaddingBits = BinaryStream.ReadUInt16();
 			NetworkByteOrderIP = BinaryStream.ReadUInt32();
 		}
 
@@ -188,6 +204,11 @@ namespace NetworkProfiler
 			ReturnList.Add( "SOCKET SEND TO" );
 			ReturnList.Add( "   ThreadId      : " + ThreadId );
 			ReturnList.Add( "   SocketName    : " + NetworkStream.GetName(SocketNameIndex) );
+			ReturnList.Add( "   DesiredBytesSent : " + (NumPacketIdBits + NumBunchBits + NumAckBits + NumPaddingBits) / 8.0f );
+			ReturnList.Add( "      NumPacketIdBits  : " + NumPacketIdBits );
+			ReturnList.Add( "      NumBunchBits     : " + NumBunchBits );
+			ReturnList.Add( "      NumAckBits       : " + NumAckBits );
+			ReturnList.Add( "      NumPaddingBits   : " + NumPaddingBits );
 			ReturnList.Add( "   BytesSent     : " + BytesSent );
 			ReturnList.Add( "   Destination   : " + NetworkByteOrderIP );
 			return ReturnList;
@@ -203,15 +224,26 @@ namespace NetworkProfiler
 		public UInt16 ChannelIndex;
 		/** Channel type. */
 		public byte ChannelType;
-		/** Number of bits serialized/ sent. */
-		public UInt16 NumBits;
+		/** Number of header bits serialized/sent. */
+		public UInt16 NumHeaderBits;
+		/** Number of non-header bits serialized/sent. */
+		public UInt16 NumPayloadBits;
 
 		/** Constructor, serializing members from passed in stream. */
 		public TokenSendBunch(BinaryReader BinaryStream)
 		{
 			ChannelIndex = BinaryStream.ReadUInt16();
 			ChannelType = BinaryStream.ReadByte();
-			NumBits = BinaryStream.ReadUInt16();
+			NumHeaderBits = BinaryStream.ReadUInt16();
+			NumPayloadBits = BinaryStream.ReadUInt16();
+		}
+
+		/**
+		 * Gets the total number of bits serialized for the bunch.
+		 */
+		public int GetNumTotalBits()
+		{
+			return NumHeaderBits + NumPayloadBits;
 		}
 
 		/**
@@ -223,8 +255,10 @@ namespace NetworkProfiler
 			ReturnList.Add( "SEND BUNCH" );
 			ReturnList.Add( "   Channel Index : " + ChannelIndex );
 			ReturnList.Add( "   Channel Type  : " + ChannelType );
-			ReturnList.Add( "   NumBits       : " + NumBits );
-			ReturnList.Add( "   NumBytes      : " + NumBits / 8.0f );
+			ReturnList.Add( "   NumTotalBits  : " + GetNumTotalBits() );
+			ReturnList.Add( "      NumHeaderBits : " + NumHeaderBits );
+			ReturnList.Add( "      NumPayloadBits: " + NumPayloadBits );
+			ReturnList.Add( "   NumTotalBytes : " + GetNumTotalBits() / 8.0f );
 			return ReturnList;
 		}
 	}
@@ -238,15 +272,29 @@ namespace NetworkProfiler
 		public int ActorNameIndex;
 		/** Name table index of function name. */
 		public int FunctionNameIndex;
-		/** Number of bits serialized/ sent. */
-		public UInt16 NumBits;
+		/** Number of bits serialized/sent for the header. */
+		public UInt16 NumHeaderBits;
+		/** Number of bits serialized/sent for the parameters. */
+		public UInt16 NumParameterBits;
+		/** Number of bits serialized/sent for the footer. */
+		public UInt16 NumFooterBits;
 
 		/** Constructor, serializing members from passed in stream. */
 		public TokenSendRPC( BinaryReader BinaryStream )
 		{
 			ActorNameIndex = BinaryStream.ReadInt32();
 			FunctionNameIndex = BinaryStream.ReadInt32();
-			NumBits = BinaryStream.ReadUInt16(); 					
+			NumHeaderBits = BinaryStream.ReadUInt16(); 					
+			NumParameterBits = BinaryStream.ReadUInt16();
+			NumFooterBits = BinaryStream.ReadUInt16();
+		}
+
+		/**
+		 * Gets the total number of bits serialized for the RPC.
+		 */
+		public int GetNumTotalBits()
+		{
+			return NumHeaderBits + NumParameterBits + NumFooterBits;
 		}
 
 		/**
@@ -258,8 +306,11 @@ namespace NetworkProfiler
 			ReturnList.Add( "SEND RPC" );
 			ReturnList.Add( "   Actor         : " + NetworkStream.GetName(ActorNameIndex) );
 			ReturnList.Add( "   Function      : " + NetworkStream.GetName(FunctionNameIndex) );
-			ReturnList.Add( "   NumBits       : " + NumBits );
-			ReturnList.Add( "   NumBytes      : " + NumBits / 8.0f );
+			ReturnList.Add( "   NumTotalBits  : " + GetNumTotalBits() );
+			ReturnList.Add( "      NumHeaderBits    : " + NumHeaderBits );
+			ReturnList.Add( "      NumParameterBits : " + NumParameterBits );
+			ReturnList.Add( "      NumFooterBits    : " + NumFooterBits );
+			ReturnList.Add( "   NumTotalBytes : " + GetNumTotalBits() / 8.0f );
 			return ReturnList;
 		}
 
@@ -463,4 +514,31 @@ namespace NetworkProfiler
 			RawData = BinaryStream.ReadBytes( Size );
 		}
 	}
+
+	/**
+	 * Token for sent acks.
+	 */
+	class TokenSendAck : TokenBase
+	{
+		/** Number of bits serialized/sent. */
+		public UInt16	NumBits;
+
+		/** Constructor, serializing members from passed in stream. */
+		public TokenSendAck(BinaryReader BinaryStream)
+		{
+			NumBits = BinaryStream.ReadUInt16();
+		}
+
+		/**
+		 * Converts the token into a multi-string description.
+		 */
+		public override List<string> ToDetailedStringList( string ActorFilter, string PropertyFilter, string RPCFilter )
+		{
+			var ReturnList = new List<string>();
+			ReturnList.Add( "SEND ACK" );
+			ReturnList.Add( "   NumBits  : " + NumBits );
+			ReturnList.Add( "   NumBytes : " + NumBits / 8.0f );
+			return ReturnList;
+		}
+	}
 }
\ No newline at end of file
diff --git a/Engine/Source/Runtime/Engine/Classes/Engine/NetConnection.h b/Engine/Source/Runtime/Engine/Classes/Engine/NetConnection.h
--- a/Engine/Source/Runtime/Engine/Classes/Engine/NetConnection.h
+++ b/Engine/Source/Runtime/Engine/Classes/Engine/NetConnection.h
@@ -137,6 +137,32 @@ class UNetConnection : public UPlayer
 	uint32 InternalAck:1;					// Internally ack all packets, for 100% reliable connections.
 
 	struct FURL			URL;				// URL of the other side.
+
+	// Track each type of bit used per-packet for bandwidth profiling
+
+	/** Number of bits used for the packet id in the current packet. */
+	int NumPacketIdBits;
+
+	/** Number of bits used for bunches in the current packet. */
+	int NumBunchBits;
+
+	/** Number of bits used for acks in the current packet. */
+	int NumAckBits;
+
+	/** Number of bits used for padding in the current packet. */
+	int NumPaddingBits;
+
+	/** Sets all of the bit-tracking variables to zero. */
+	void ResetPacketBitCounts();
+
+	/** What type of data is being written */
+	enum class EWriteBitsDataType
+	{
+		Unknown,
+		Bunch,
+		Ack
+	};
+
 public:
 	// Constants.
 	enum{ MAX_PROTOCOL_VERSION = 1     };	// Maximum protocol version supported.
@@ -466,16 +492,21 @@ public:
 	 * @param SizeInBits Number of bits to append
 	 * @param ExtraBits (optional) Second set of bits to be appended to the send buffer that need to send with the first set of bits
 	 * @param ExtraSizeInBits (optional) Number of secondary bits to append
+	 * @param TypeOfBits (optional) The type of data being written, for profiling and bandwidth tracking purposes
 	 */
 	int32 WriteBitsToSendBuffer( 
 		const uint8 *	Bits, 
 		const int32		SizeInBits, 
 		const uint8 *	ExtraBits = NULL, 
-		const int32		ExtraSizeInBits = 0 );
+		const int32		ExtraSizeInBits = 0,
+		EWriteBitsDataType DataType =  EWriteBitsDataType::Unknown);
 
 	/** Returns number of bits left in current packet that can be used without causing a flush.  */
 	int64 GetFreeSendBufferBits();
 
+	/** Pops the LastStart bits off of the send buffer, used for merging bunches */
+	void PopLastStart();
+
 	/** 
 	 * returns whether the client has initialized the level required for the given object
 	 * @return true if the client has initialized the level the object is in or the object is not in a level, false otherwise
diff --git a/Engine/Source/Runtime/Engine/Private/DataChannel.cpp b/Engine/Source/Runtime/Engine/Private/DataChannel.cpp
--- a/Engine/Source/Runtime/Engine/Private/DataChannel.cpp
+++ b/Engine/Source/Runtime/Engine/Private/DataChannel.cpp
@@ -730,7 +730,7 @@ FPacketIdRange UChannel::SendBunch( FOutBunch* Bunch, bool Merge )
 		OutBunch                       = Connection->LastOutBunch;
 		Bunch                          = &Connection->LastOut;
 		check(!Bunch->IsError());
-		Connection->LastStart.Pop( Connection->SendBuffer );
+		Connection->PopLastStart();
 		Connection->Driver->OutBunches--;
 	}
 
@@ -822,8 +822,6 @@ FPacketIdRange UChannel::SendBunch( FOutBunch* Bunch, bool Merge )
 			UE_LOG(LogNetPartialBunch, Verbose, TEXT("	Bunch[%d]: Bytes: %d Bits: %d ChSequence: %d 0x%X"), PartialNum, ThisOutBunch->GetNumBytes(), ThisOutBunch->GetNumBits(), ThisOutBunch->ChSequence, FCrc::MemCrc_DEPRECATED(ThisOutBunch->GetData(), ThisOutBunch->GetNumBytes()));
 		}
 
-		NETWORK_PROFILER(GNetworkProfiler.TrackSendBunch(ThisOutBunch,ThisOutBunch->GetNumBits()-PreExistingBits));
-
 		// Update Packet Range
 		int32 PacketId = SendRawBunch(ThisOutBunch, Merge);
 		if (PartialNum == 0)
diff --git a/Engine/Source/Runtime/Engine/Private/NetConnection.cpp b/Engine/Source/Runtime/Engine/Private/NetConnection.cpp
--- a/Engine/Source/Runtime/Engine/Private/NetConnection.cpp
+++ b/Engine/Source/Runtime/Engine/Private/NetConnection.cpp
@@ -425,6 +425,8 @@ void UNetConnection::InitSendBuffer()
 		SendBuffer = FBitWriter(MaxPacket * 8);
 	}
 
+	ResetPacketBitCounts();
+
 	ValidateSendBuffer();
 }
 
@@ -482,6 +484,8 @@ void UNetConnection::FlushNet(bool bIgnoreSimulation)
 			WriteBitsToSendBuffer( NULL, 0 );		// This will force the packet id to be written
 		}
 
+		const int NumBitsPrePadding = SendBuffer.GetNumBits();
+
 		// Make sure packet size is byte-aligned.
 		SendBuffer.WriteBit( 1 );
 		while( SendBuffer.GetNumBits() & 7 )
@@ -490,6 +494,10 @@ void UNetConnection::FlushNet(bool bIgnoreSimulation)
 		}
 		ValidateSendBuffer();
 
+		NumPaddingBits += SendBuffer.GetNumBits() - NumBitsPrePadding;
+
+		NETWORK_PROFILER(GNetworkProfiler.FlushOutgoingBunches(this));
+
 		// Send now.
 #if DO_ENABLE_NET_TEST
 		// if the connection is closing/being destroyed/etc we need to send immediately regardless of settings
@@ -1006,7 +1014,8 @@ int32 UNetConnection::WriteBitsToSendBuffer(
 	const uint8 *	Bits, 
 	const int32		SizeInBits, 
 	const uint8 *	ExtraBits, 
-	const int32		ExtraSizeInBits )
+	const int32		ExtraSizeInBits,
+	EWriteBitsDataType DataType)
 {
 	ValidateSendBuffer();
 
@@ -1027,6 +1036,8 @@ int32 UNetConnection::WriteBitsToSendBuffer(
 	{
 		SendBuffer.WriteIntWrapped( OutPacketId, MAX_PACKETID );
 		ValidateSendBuffer();
+
+		NumPacketIdBits += SendBuffer.GetNumBits();
 	}
 
 	// Add the bits to the queue
@@ -1045,6 +1056,18 @@ int32 UNetConnection::WriteBitsToSendBuffer(
 
 	const int32 RememberedPacketId = OutPacketId;
 
+	switch ( DataType )
+	{
+		case EWriteBitsDataType::Bunch:
+			NumBunchBits += SizeInBits + ExtraSizeInBits;
+			break;
+		case EWriteBitsDataType::Ack:
+			NumAckBits += SizeInBits + ExtraSizeInBits;
+			break;
+		default:
+			break;
+	}
+
 	// Flush now if we are full
 	if ( GetFreeSendBufferBits() == 0 )
 	{
@@ -1068,6 +1091,13 @@ int64 UNetConnection::GetFreeSendBufferBits()
 	return NumberOfFreeBits;
 }
 
+void UNetConnection::PopLastStart()
+{
+	NumBunchBits -= SendBuffer.GetNumBits() - LastStart.GetNumBits();
+	LastStart.Pop(SendBuffer);
+	NETWORK_PROFILER(GNetworkProfiler.PopSendBunch(this));
+}
+
 void UNetConnection::PurgeAcks()
 {
 	for ( int32 i = 0; i < ResendAcks.Num(); i++ )
@@ -1108,7 +1138,9 @@ void UNetConnection::SendAck( int32 AckPacketId, bool FirstTime/*=1*/, bool bHav
 			}
 		}
 
-		WriteBitsToSendBuffer( AckData.GetData(), AckData.GetNumBits() );
+		NETWORK_PROFILER( GNetworkProfiler.TrackSendAck( AckData.GetNumBits() ) );
+
+		WriteBitsToSendBuffer( AckData.GetData(), AckData.GetNumBits(), nullptr, 0, EWriteBitsDataType::Ack );
 
 		AllowMerge = false;
 
@@ -1177,8 +1209,10 @@ int32 UNetConnection::SendRawBunch( FOutBunch& Bunch, bool InAllowMerge )
 		UE_LOG(LogNetTraffic, VeryVerbose, TEXT("Sending: %s"), *Bunch.ToString());
 	}
 
+	NETWORK_PROFILER(GNetworkProfiler.PushSendBunch(this, &Bunch, Header.GetNumBits(), Bunch.GetNumBits()));
+
 	// Write the bits to the buffer and remember the packet id used
-	Bunch.PacketId = WriteBitsToSendBuffer( Header.GetData(), Header.GetNumBits(), Bunch.GetData(), Bunch.GetNumBits() );
+	Bunch.PacketId = WriteBitsToSendBuffer( Header.GetData(), Header.GetNumBits(), Bunch.GetData(), Bunch.GetNumBits(), EWriteBitsDataType::Bunch );
 
 	UE_LOG(LogNetTraffic, Verbose, TEXT("UNetConnection::SendRawBunch. ChIndex: %d. Bits: %d. PacketId: %d"), Bunch.ChIndex, Bunch.GetNumBits(), Bunch.PacketId );
 
@@ -1832,3 +1866,10 @@ bool UNetConnection::TrackLogsPerSecond()
 	return true;
 }
 
+void UNetConnection::ResetPacketBitCounts()
+{
+	NumPacketIdBits = 0;
+	NumBunchBits = 0;
+	NumAckBits = 0;
+	NumPaddingBits = 0;
+}
\ No newline at end of file
diff --git a/Engine/Source/Runtime/Engine/Private/NetworkDriver.cpp b/Engine/Source/Runtime/Engine/Private/NetworkDriver.cpp
--- a/Engine/Source/Runtime/Engine/Private/NetworkDriver.cpp
+++ b/Engine/Source/Runtime/Engine/Private/NetworkDriver.cpp
@@ -833,6 +833,8 @@ void UNetDriver::InternalProcessRemoteFunction
 	check( FieldCache->FieldNetIndex <= ClassCache->GetMaxIndex() );
 	Bunch.WriteIntWrapped(FieldCache->FieldNetIndex, ClassCache->GetMaxIndex()+1);
 
+	const int HeaderBits = Bunch.GetNumBits();
+
 #if !(UE_BUILD_SHIPPING || UE_BUILD_TEST)
 	Bunch.DebugString = FString::Printf(TEXT("%.2f RPC: %s - %s"), Connection->Driver->Time, *Actor->GetName(), *Function->GetName());
 #endif
@@ -930,6 +932,8 @@ void UNetDriver::InternalProcessRemoteFunction
 	TSharedPtr<FRepLayout> RepLayout = GetFunctionRepLayout( Function );
 	RepLayout->SendPropertiesForRPC( Actor, Function, Ch, Bunch, Parms );
 
+	const int ParameterBits = Bunch.GetNumBits() - HeaderBits;
+
 	// Destroy the memory used for the copied out parameters
 	for ( int32 i = 0; i < LocalOutParms.Num(); i++ )
 	{
@@ -943,6 +947,8 @@ void UNetDriver::InternalProcessRemoteFunction
 		Ch->EndContentBlock( TargetObj, Bunch );
 	}
 
+	const int FooterBits = Bunch.GetNumBits() - HeaderBits - ParameterBits;
+
 	static const auto CVar = IConsoleManager::Get().FindTConsoleVariableDataInt(TEXT("net.RPC.Debug"));
 	bool LogAsWarning = (CVar && CVar->GetValueOnGameThread() == 1);
 
@@ -981,7 +987,10 @@ void UNetDriver::InternalProcessRemoteFunction
 			UE_LOG(LogNetTraffic, Log,		TEXT("      Sent RPC: %s::%s [%.1f bytes]"), *Actor->GetName(), *Function->GetName(), Bunch.GetNumBits() / 8.f );
 		}
 
-		NETWORK_PROFILER(GNetworkProfiler.TrackSendRPC(Actor,Function,Bunch.GetNumBits()));
+		// Make sure we're tracking all the bits in the bunch
+		check(Bunch.GetNumBits() == HeaderBits + ParameterBits + FooterBits);
+
+		NETWORK_PROFILER(GNetworkProfiler.TrackSendRPC(Actor, Function, HeaderBits, ParameterBits, FooterBits));
 		Ch->SendBunch( &Bunch, 1 );
 	}
 }
diff --git a/Engine/Source/Runtime/Engine/Private/NetworkProfiler.cpp b/Engine/Source/Runtime/Engine/Private/NetworkProfiler.cpp
--- a/Engine/Source/Runtime/Engine/Private/NetworkProfiler.cpp
+++ b/Engine/Source/Runtime/Engine/Private/NetworkProfiler.cpp
@@ -26,7 +26,7 @@ FNetworkProfiler GNetworkProfiler;
 /** Magic value, determining that file is a network profiler file.				*/
 #define NETWORK_PROFILER_MAGIC						0x1DBF348C
 /** Version of memory profiler. Incremented on serialization changes.			*/
-#define NETWORK_PROFILER_VERSION					5
+#define NETWORK_PROFILER_VERSION					6
 
 enum ENetworkProfilingPayloadType
 {
@@ -39,6 +39,7 @@ enum ENetworkProfilingPayloadType
 	NPTYPE_EndOfStreamMarker,			// End of stream marker		
 	NPTYPE_Event,						// Event
 	NPTYPE_RawSocketData,				// Raw socket data being sent
+	NPTYPE_SendAck						// Ack being sent
 };
 
 
@@ -174,7 +175,7 @@ void FNetworkProfiler::TrackFrameBegin()
  * @param	Function	Function being called
  * @param	NumBits		Number of bits serialized into bunch for this RPC
  */
-void FNetworkProfiler::TrackSendRPC( const AActor* Actor, const UFunction* Function, uint16 NumBits )
+void FNetworkProfiler::TrackSendRPC( const AActor* Actor, const UFunction* Function, uint16 NumHeaderBits, uint16 NumParameterBits, uint16 NumFooterBits  )
 {
 	if( bIsTrackingEnabled )
 	{
@@ -185,7 +186,9 @@ void FNetworkProfiler::TrackSendRPC( const AActor* Actor, const UFunction* Funct
 		(*FileWriter) << ActorNameTableIndex;
 		int32 FunctionNameTableIndex = GetNameTableIndex( Function->GetName() );
 		(*FileWriter) << FunctionNameTableIndex;
-		(*FileWriter) << NumBits;
+		(*FileWriter) << NumHeaderBits;
+		(*FileWriter) << NumParameterBits;
+		(*FileWriter) << NumFooterBits;
 	}
 }
 
@@ -202,7 +205,7 @@ void FNetworkProfiler::TrackSocketSend( const FString& SocketDesc, const void* D
 	{
 		SCOPE_LOCK_REF(CriticalSection);
 		uint32 DummyIp = 0;
-		TrackSocketSendToCore( SocketDesc, Data, BytesSent, DummyIp );
+		//TrackSocketSendToCore( SocketDesc, Data, BytesSent, DummyIp );
 	}
 }
 
@@ -214,13 +217,21 @@ void FNetworkProfiler::TrackSocketSend( const FString& SocketDesc, const void* D
  * @param	BytesSent				Bytes actually being sent
  * @param	Destination				Destination address
  */
-void FNetworkProfiler::TrackSocketSendTo( const FString& SocketDesc, const void* Data, uint16 BytesSent, const FInternetAddr& Destination )
+void FNetworkProfiler::TrackSocketSendTo(
+	const FString& SocketDesc,
+	const void* Data,
+	uint16 BytesSent,
+	uint16 NumPacketIdBits,
+	uint16 NumBunchBits,
+	uint16 NumAckBits,
+	uint16 NumPaddingBits,
+	const FInternetAddr& Destination )
 {
 	if( bIsTrackingEnabled )
 	{
 		uint32 NetworkByteOrderIP;
 		Destination.GetIp(NetworkByteOrderIP);
-		TrackSocketSendToCore( SocketDesc, Data, BytesSent, NetworkByteOrderIP);
+		TrackSocketSendToCore( SocketDesc, Data, BytesSent, NumPacketIdBits, NumBunchBits, NumAckBits, NumPaddingBits, NetworkByteOrderIP);
 		bHasNoticeableNetworkTrafficOccured = true;
 	}
 }
@@ -233,7 +244,15 @@ void FNetworkProfiler::TrackSocketSendTo( const FString& SocketDesc, const void*
  * @param	BytesSent				Bytes actually being sent
  * @param	IpAddr					Destination address
  */
-void FNetworkProfiler::TrackSocketSendToCore( const FString& SocketDesc, const void* Data, uint16 BytesSent, uint32 IpAddr )
+void FNetworkProfiler::TrackSocketSendToCore(
+	const FString& SocketDesc,
+	const void* Data,
+	uint16 BytesSent,
+	uint16 NumPacketIdBits,
+	uint16 NumBunchBits,
+	uint16 NumAckBits,
+	uint16 NumPaddingBits,
+	uint32 IpAddr )
 {
 	if( bIsTrackingEnabled )
 	{
@@ -246,6 +265,10 @@ void FNetworkProfiler::TrackSocketSendToCore( const FString& SocketDesc, const v
 		int32 NameTableIndex = GetNameTableIndex( SocketDesc );
 		(*FileWriter) << NameTableIndex;
 		(*FileWriter) << BytesSent;
+		(*FileWriter) << NumPacketIdBits;
+		(*FileWriter) << NumBunchBits;
+		(*FileWriter) << NumAckBits;
+		(*FileWriter) << NumPaddingBits;
 		(*FileWriter) << IpAddr;
 #if NETWORK_PROFILER_TRACK_RAW_NETWORK_DATA
 		Type = NPTYPE_RawSocketData;
@@ -279,6 +302,49 @@ void FNetworkProfiler::TrackSendBunch( FOutBunch* OutBunch, uint16 NumBits )
 	}
 }
 
+void FNetworkProfiler::PushSendBunch( UNetConnection* Connection, FOutBunch* OutBunch, uint16 NumHeaderBits, uint16 NumPayloadBits )
+{
+	if ( bIsTrackingEnabled )
+	{
+		SCOPE_LOCK_REF(CriticalSection);
+		OutgoingBunches.FindOrAdd(Connection).Emplace(OutBunch->ChIndex, OutBunch->ChType, NumHeaderBits, NumPayloadBits);
+	}
+}
+
+void FNetworkProfiler::PopSendBunch( UNetConnection* Connection )
+{
+	if ( bIsTrackingEnabled )
+	{
+		SCOPE_LOCK_REF(CriticalSection);
+		if ( OutgoingBunches.Contains(Connection) && OutgoingBunches[Connection].Num() > 0 )
+		{
+			OutgoingBunches[Connection].Pop();
+		}
+	}
+}
+
+void FNetworkProfiler::FlushOutgoingBunches( UNetConnection* Connection )
+{
+	if ( bIsTrackingEnabled )
+	{
+		SCOPE_LOCK_REF(CriticalSection);
+		if ( OutgoingBunches.Contains(Connection) )
+		{
+			for ( FSendBunchInfo& BunchInfo : OutgoingBunches[Connection] )
+			{
+				uint8 Type = NPTYPE_SendBunch;
+				(*FileWriter) << Type;
+				(*FileWriter) << BunchInfo.ChannelIndex;
+				(*FileWriter) << BunchInfo.ChannelType;
+				(*FileWriter) << BunchInfo.NumHeaderBits;
+				(*FileWriter) << BunchInfo.NumPayloadBits;
+			}
+
+			OutgoingBunches[Connection].SetNum(0);
+		}
+	}
+}
+
 /**
  * Track actor being replicated.
  *
@@ -448,6 +514,17 @@ void FNetworkProfiler::TrackSessionChange( bool bShouldContinueTracking, const F
 #endif	//#if ALLOW_DEBUG_FILES
 }
 
+ void FNetworkProfiler::TrackSendAck( uint16 NumBits )
+ {
+	 if ( bIsTrackingEnabled )
+	 {
+		 SCOPE_LOCK_REF(CriticalSection);
+		 uint8 Type = NPTYPE_SendAck;
+		 (*FileWriter) << Type;
+		 (*FileWriter) << NumBits;
+	 }
+ }
+
 /**
  * Processes any network profiler specific exec commands
  *
diff --git a/Engine/Source/Runtime/Engine/Public/Net/NetworkProfiler.h b/Engine/Source/Runtime/Engine/Public/Net/NetworkProfiler.h
--- a/Engine/Source/Runtime/Engine/Public/Net/NetworkProfiler.h
+++ b/Engine/Source/Runtime/Engine/Public/Net/NetworkProfiler.h
@@ -7,6 +7,8 @@
 #ifndef UNREAL_NETWORK_PROFILER_H
 #define UNREAL_NETWORK_PROFILER_H
 
+class UNetConnection;
+
 #if USE_NETWORK_PROFILER 
 
 #define NETWORK_PROFILER( x ) if ( GNetworkProfiler.IsTrackingEnabled() ) { x; }
@@ -42,6 +44,30 @@ private:
 
 	/** URL used for current tracking session.														*/
 	FURL									CurrentURL;
+	
+	/** All the data required for writing sent bunches to the profiler stream						*/
+	struct FSendBunchInfo
+	{
+		uint16 ChannelIndex;
+		uint8 ChannelType;
+		uint16 NumHeaderBits;
+		uint16 NumPayloadBits;
+
+		FSendBunchInfo()
+			: ChannelIndex(0)
+			, ChannelType(0)
+			, NumHeaderBits(0)
+			, NumPayloadBits(0) {}
+
+		FSendBunchInfo( uint16 InChannelIndex, uint8 InChannelType, uint16 InNumHeaderBits, uint16 InNumPayloadBits )
+			: ChannelIndex(InChannelIndex)
+			, ChannelType(InChannelType)
+			, NumHeaderBits(InNumHeaderBits)
+			, NumPayloadBits(InNumPayloadBits) {}
+	};
+
+	/** Stack outgoing bunches per connection, the top bunch for a connection may be popped if it gets merged with a new bunch.		*/
+	TMap<UNetConnection*, TArray<FSendBunchInfo>>	OutgoingBunches;
 
 	/**
 	 * Returns index of passed in name into name array. If not found, adds it.
@@ -72,11 +98,13 @@ public:
 	/**
 	 * Tracks and RPC being sent.
 	 * 
-	 * @param	Actor		Actor RPC is being called on
-	 * @param	Function	Function being called
-	 * @param	NumBits		Number of bits serialized into bunch for this RPC
+	 * @param	Actor				Actor RPC is being called on
+	 * @param	Function			Function being called
+	 * @param	NumHeaderBits		Number of bits serialized into the header for this RPC
+	 * @param	NumParameterBits	Number of bits serialized into parameters of this RPC
+	 * @param	NumFooterBits		Number of bits serialized into the footer of this RPC (EndContentBlock)
 	 */
-	void TrackSendRPC( const AActor* Actor, const UFunction* Function, uint16 NumBits );
+	void TrackSendRPC( const AActor* Actor, const UFunction* Function, uint16 NumHeaderBits, uint16 NumParameterBits, uint16 NumFooterBits );
 	
 	/**
 	 * Low level FSocket::Send information.
@@ -93,9 +121,21 @@ public:
  	 * @param	SocketDesc				Description of socket data is being sent to
 	 * @param	Data					Data sent
 	 * @param	BytesSent				Bytes actually being sent
+	 * @param	NumPacketIdBits			Number of bits sent for the packet id
+	 * @param	NumBunchBits			Number of bits sent in bunches
+	 * @param	NumAckBits				Number of bits sent in acks
+	 * @param	NumPaddingBits			Number of bits appended to the end to byte-align the data
 	 * @param	Destination				Destination address
 	 */
-	ENGINE_API void TrackSocketSendTo( const FString& SocketDesc, const void* Data, uint16 BytesSent, const FInternetAddr& Destination );
+	ENGINE_API void TrackSocketSendTo(
+		const FString& SocketDesc,
+		const void* Data,
+		uint16 BytesSent,
+		uint16 NumPacketIdBits,
+		uint16 NumBunchBits,
+		uint16 NumAckBits,
+		uint16 NumPaddingBits,
+		const FInternetAddr& Destination );
 
 	/**
 	 * Low level FSocket::SendTo information.
@@ -105,7 +145,15 @@ public:
 	 * @param	BytesSent				Bytes actually being sent
 	 * @param	IpAddr					Destination address
 	 */
-	void TrackSocketSendToCore( const FString& SocketDesc, const void* Data, uint16 BytesSent, uint32 IpAddr );
+	void TrackSocketSendToCore(
+		const FString& SocketDesc,
+		const void* Data,
+		uint16 BytesSent,
+		uint16 NumPacketIdBits,
+		uint16 NumBunchBits,
+		uint16 NumAckBits,
+		uint16 NumPaddingBits,
+		uint32 IpAddr );
 
 	
 	/**
@@ -116,6 +164,31 @@ public:
 	 */
 	void TrackSendBunch( FOutBunch* OutBunch, uint16 NumBits );
 	
+	/**
+	 * Add a sent bunch to the stack. These bunches are not written to the stream immediately,
+	 * because they may be merged with another bunch in the future.
+	 *
+	 * @param Connection The connection on which this bunch was sent
+	 * @param OutBunch The bunch being sent
+	 * @param NumHeaderBits Number of bits in the bunch header
+	 * @param NumPayloadBits Number of bits in the bunch, excluding the header
+	 */
+	void PushSendBunch( UNetConnection* Connection, FOutBunch* OutBunch, uint16 NumHeaderBits, uint16 NumPayloadBits );
+
+	/**
+	 * Pops the latest bunch for a connection, since it is going to be merged with the next bunch.
+	 *
+	 * @param Connection the connection which is merging a bunch
+	 */
+	void PopSendBunch( UNetConnection* Connection );
+
+	/**
+	 * Writes all the outgoing bunches for a connection in the stack to the profiler data stream.
+	 *
+	 * @param Connection the connection which is about to send any pending bunches over the network
+	 */
+	void FlushOutgoingBunches( UNetConnection* Connection );
+
 	/**
 	 * Track actor being replicated.
 	 *
@@ -149,6 +222,13 @@ public:
 	 */
 	void TrackSessionChange( bool bShouldContinueTracking, const FURL& InURL );
 
+	/**
+	 * Track sent acks.
+	 *
+	 * @param NumBits Number of bits in the ack
+	 */
+	void TrackSendAck( uint16 NumBits );
+
 	/**
 	 * Processes any network profiler specific exec commands
 	 *
diff --git a/Engine/Source/Runtime/Online/OnlineSubsystemUtils/Private/IpConnection.cpp b/Engine/Source/Runtime/Online/OnlineSubsystemUtils/Private/IpConnection.cpp
--- a/Engine/Source/Runtime/Online/OnlineSubsystemUtils/Private/IpConnection.cpp
+++ b/Engine/Source/Runtime/Online/OnlineSubsystemUtils/Private/IpConnection.cpp
@@ -133,7 +133,7 @@ void UIpConnection::LowLevelSend( void* Data, int32 Count )
 	CLOCK_CYCLES(Driver->SendCycles);
 	Socket->SendTo((uint8*)Data, Count, BytesSent, *RemoteAddr);
 	UNCLOCK_CYCLES(Driver->SendCycles);
-	NETWORK_PROFILER(GNetworkProfiler.TrackSocketSendTo(Socket->GetDescription(),Data,BytesSent,*RemoteAddr));
+	NETWORK_PROFILER(GNetworkProfiler.TrackSocketSendTo(Socket->GetDescription(),Data,BytesSent,NumPacketIdBits,NumBunchBits,NumAckBits,NumPaddingBits,*RemoteAddr));
 }
 
 FString UIpConnection::LowLevelGetRemoteAddress(bool bAppendPort)
