Track more bandwidth information in the network profiler, and show a more useful chart view by default.
"Custom delta" replicated properties are now tracked (this includes fast TArrays).
More of the bunch format metadata is also tracked: property headers, property handles, exported GUIDs, "must be mapped" GUIDs, and content block headers & footers.
The chart can now show a stacked area view of bunch data, so that it's easier to see at a glance what's taking up space in bunches. Note that for large profiles this view may be slower to render.

diff --git a/Engine/Source/Programs/NetworkProfiler/NetworkProfiler/Charting.cs b/Engine/Source/Programs/NetworkProfiler/NetworkProfiler/Charting.cs
--- a/Engine/Source/Programs/NetworkProfiler/NetworkProfiler/Charting.cs
+++ b/Engine/Source/Programs/NetworkProfiler/NetworkProfiler/Charting.cs
@@ -63,14 +63,22 @@ namespace NetworkProfiler
 				NetworkChart.Series["RPCCountSec"].Points.AddXY( FrameCounter, Frame.RPCCount * OneOverDeltaTime );
 				NetworkChart.Series["RPCSize"].Points.AddXY( FrameCounter, Frame.RPCSizeBits / 8 );
 				NetworkChart.Series["RPCSizeSec"].Points.AddXY( FrameCounter, Frame.RPCSizeBits / 8 * OneOverDeltaTime );
+				NetworkChart.Series["ExportBunchCount"].Points.AddXY( FrameCounter, Frame.ExportBunchCount );
+				NetworkChart.Series["ExportBunchSize"].Points.AddXY( FrameCounter, Frame.ExportBunchSizeBits / 8 );
+				NetworkChart.Series["MustBeMappedGuidsCount"].Points.AddXY( FrameCounter, Frame.MustBeMappedGuidCount / 8 );
+				NetworkChart.Series["MustBeMappedGuidsSize"].Points.AddXY( FrameCounter, Frame.MustBeMappedGuidSizeBits / 8 );
 				NetworkChart.Series["SendAckCount"].Points.AddXY( FrameCounter, Frame.SendAckCount );
 				NetworkChart.Series["SendAckCountSec"].Points.AddXY( FrameCounter, Frame.SendAckCount * OneOverDeltaTime );
 				NetworkChart.Series["SendAckSize"].Points.AddXY( FrameCounter, Frame.SendAckSizeBits / 8 );
 				NetworkChart.Series["SendAckSizeSec"].Points.AddXY( FrameCounter, Frame.SendAckSizeBits / 8 * OneOverDeltaTime );
+				NetworkChart.Series["ContentBlockHeaderSize"].Points.AddXY( FrameCounter, Frame.ContentBlockHeaderSizeBits / 8 );
+				NetworkChart.Series["ContentBlockFooterSize"].Points.AddXY( FrameCounter, Frame.ContentBlockFooterSizeBits / 8 );
+				NetworkChart.Series["PropertyHandleSize"].Points.AddXY( FrameCounter, Frame.PropertyHandleSizeBits / 8 );
 				NetworkChart.Series["SendBunchCount"].Points.AddXY( FrameCounter, Frame.SendBunchCount );
 				NetworkChart.Series["SendBunchCountSec"].Points.AddXY( FrameCounter, Frame.SendBunchCount * OneOverDeltaTime );
 				NetworkChart.Series["SendBunchSize"].Points.AddXY( FrameCounter, Frame.SendBunchSizeBits / 8 );
 				NetworkChart.Series["SendBunchSizeSec"].Points.AddXY( FrameCounter, Frame.SendBunchSizeBits / 8 * OneOverDeltaTime );
+				NetworkChart.Series["SendBunchHeaderSize"].Points.AddXY( FrameCounter, Frame.SendBunchHeaderSizeBits / 8 );
 				NetworkChart.Series["GameSocketSendSize"].Points.AddXY( FrameCounter, Frame.UnrealSocketSize );
 				NetworkChart.Series["GameSocketSendSizeSec"].Points.AddXY( FrameCounter, Frame.UnrealSocketSize * OneOverDeltaTime );
 				NetworkChart.Series["GameSocketSendCount"].Points.AddXY( FrameCounter, Frame.UnrealSocketCount );
diff --git a/Engine/Source/Programs/NetworkProfiler/NetworkProfiler/MainWindow.cs b/Engine/Source/Programs/NetworkProfiler/NetworkProfiler/MainWindow.cs
--- a/Engine/Source/Programs/NetworkProfiler/NetworkProfiler/MainWindow.cs
+++ b/Engine/Source/Programs/NetworkProfiler/NetworkProfiler/MainWindow.cs
@@ -30,22 +30,93 @@ namespace NetworkProfiler
 		{
 			InitializeComponent();
 
+			// Set default state.
+			SetDefaultLineView();
+
+			// Force the columns to the way we want them
+			SetupColumns( ActorListView,	new String[]	{ "Total Size (KBytes)", "Count", "Average Size (Bytes)", "Average Size (Bits)", "Time (ms)", "Actor Class" });
+			SetupColumns( PropertyListView, new String[]	{ "Total Size (KBytes)", "Count", "Average Size (Bytes)", "Average Size (Bits)", "Time (ms)", "Property" });
+			SetupColumns( RPCListView,		new String[]	{ "Total Size (KBytes)", "Count", "Average Size (Bytes)", "Average Size (Bits)", "Time (ms)", "RPC" });
+		}
+
+		private void SetDefaultLineView()
+		{
 			// Set default state.
 			foreach( var Series in NetworkChart.Series )
 			{
 				Series.Enabled = false;
 			}
+
+			for( int i = 0; i < ChartListBox.Items.Count; ++i )
+			{
+				ChartListBox.SetItemChecked(i, false);
+			}
+
 			NetworkChart.Series["PropertySize"].Enabled = true;
-			ChartListBox.SetItemChecked(4, true);
+			NetworkChart.Series["PropertySize"].ChartType = SeriesChartType.FastLine;
+			ChartListBox.SetItemChecked(NetworkChart.Series.IndexOf("PropertySize"), true);
+
 			NetworkChart.Series["RPCSize"].Enabled = true;
-			ChartListBox.SetItemChecked(8, true);
+			NetworkChart.Series["RPCSize"].ChartType = SeriesChartType.FastLine;
+			ChartListBox.SetItemChecked(NetworkChart.Series.IndexOf("RPCSize"), true);
+
 			NetworkChart.Series["Events"].Enabled = true;
-			ChartListBox.SetItemChecked(27, true);
+			ChartListBox.SetItemChecked(NetworkChart.Series.IndexOf("Events"), true);
 
-			// Force the columns to the way we want them
-			SetupColumns( ActorListView,	new String[]	{ "Total Size (KBytes)", "Count", "Average Size (Bytes)", "Average Size (Bits)", "Time (ms)", "Actor Class" });
-			SetupColumns( PropertyListView, new String[]	{ "Total Size (KBytes)", "Count", "Average Size (Bytes)", "Average Size (Bits)", "Time (ms)", "Property" });
-			SetupColumns( RPCListView,		new String[]	{ "Total Size (KBytes)", "Count", "Average Size (Bytes)", "Average Size (Bits)", "Time (ms)", "RPC" });
+			NetworkChart.Series["ExportBunchSize"].ChartType = SeriesChartType.FastLine;
+			NetworkChart.Series["MustBeMappedGuidsSize"].ChartType = SeriesChartType.FastLine;
+			NetworkChart.Series["SendBunchHeaderSize"].ChartType = SeriesChartType.FastLine;
+			NetworkChart.Series["ContentBlockHeaderSize"].ChartType = SeriesChartType.FastLine;
+			NetworkChart.Series["ContentBlockFooterSize"].ChartType = SeriesChartType.FastLine;
+			NetworkChart.Series["PropertyHandleSize"].ChartType = SeriesChartType.FastLine;
+		}
+
+		private void SetDefaultStackedBunchSizeView()
+		{
+			foreach( var Series in NetworkChart.Series )
+			{
+				Series.Enabled = false;
+			}
+
+			for( int i = 0; i < ChartListBox.Items.Count; ++i )
+			{
+				ChartListBox.SetItemChecked(i, false);
+			}
+
+			NetworkChart.Series["SendBunchSize"].Enabled = true;
+			ChartListBox.SetItemChecked(NetworkChart.Series.IndexOf("SendBunchSize"), true);
+
+			NetworkChart.Series["PropertySize"].Enabled = true;
+			NetworkChart.Series["PropertySize"].ChartType = SeriesChartType.StackedArea;
+			ChartListBox.SetItemChecked(NetworkChart.Series.IndexOf("PropertySize"), true);
+
+			NetworkChart.Series["RPCSize"].Enabled = true;
+			NetworkChart.Series["RPCSize"].ChartType = SeriesChartType.StackedArea;
+			ChartListBox.SetItemChecked(NetworkChart.Series.IndexOf("RPCSize"), true);
+
+			NetworkChart.Series["ExportBunchSize"].Enabled = true;
+			NetworkChart.Series["ExportBunchSize"].ChartType = SeriesChartType.StackedArea;
+			ChartListBox.SetItemChecked(NetworkChart.Series.IndexOf("ExportBunchSize"), true);
+
+			NetworkChart.Series["MustBeMappedGuidsSize"].Enabled = true;
+			NetworkChart.Series["MustBeMappedGuidsSize"].ChartType = SeriesChartType.StackedArea;
+			ChartListBox.SetItemChecked(NetworkChart.Series.IndexOf("MustBeMappedGuidsSize"), true);
+
+			NetworkChart.Series["SendBunchHeaderSize"].Enabled = true;
+			NetworkChart.Series["SendBunchHeaderSize"].ChartType = SeriesChartType.StackedArea;
+			ChartListBox.SetItemChecked(NetworkChart.Series.IndexOf("SendBunchHeaderSize"), true);
+
+			NetworkChart.Series["ContentBlockHeaderSize"].Enabled = true;
+			NetworkChart.Series["ContentBlockHeaderSize"].ChartType = SeriesChartType.StackedArea;
+			ChartListBox.SetItemChecked(NetworkChart.Series.IndexOf("ContentBlockHeaderSize"), true);
+
+			NetworkChart.Series["ContentBlockFooterSize"].Enabled = true;
+			NetworkChart.Series["ContentBlockFooterSize"].ChartType = SeriesChartType.StackedArea;
+			ChartListBox.SetItemChecked(NetworkChart.Series.IndexOf("ContentBlockFooterSize"), true);
+
+			NetworkChart.Series["PropertyHandleSize"].Enabled = true;
+			NetworkChart.Series["PropertyHandleSize"].ChartType = SeriesChartType.StackedArea;
+			ChartListBox.SetItemChecked(NetworkChart.Series.IndexOf("PropertyHandleSize"), true);
 		}
 
 		private void SetupColumns(ListView ListView, String[] Headers)
@@ -294,5 +365,21 @@ namespace NetworkProfiler
 		{
 			HandleListViewSorting( e.Column, ActorListView );
 		}
+
+		private void LineViewRadioButton_CheckChanged(object sender, EventArgs e)
+		{
+			if (LineViewRadioButton.Checked)
+			{
+				SetDefaultLineView();
+			}
+		}
+
+		private void StackedBunchSizeRadioButton_CheckChanged(object sender, EventArgs e)
+		{
+			if (StackedBunchSizeRadioButton.Checked)
+			{
+				SetDefaultStackedBunchSizeView();
+			}
+		}
 	}
 }
diff --git a/Engine/Source/Programs/NetworkProfiler/NetworkProfiler/NetworkProfiler.csproj b/Engine/Source/Programs/NetworkProfiler/NetworkProfiler/NetworkProfiler.csproj
--- a/Engine/Source/Programs/NetworkProfiler/NetworkProfiler/NetworkProfiler.csproj
+++ b/Engine/Source/Programs/NetworkProfiler/NetworkProfiler/NetworkProfiler.csproj
@@ -12,10 +12,14 @@
     <AssemblyName>NetworkProfiler</AssemblyName>
     <TargetFrameworkVersion>v4.0</TargetFrameworkVersion>
     <FileAlignment>512</FileAlignment>
-    <SccProjectName>SAK</SccProjectName>
-    <SccLocalPath>SAK</SccLocalPath>
-    <SccAuxPath>SAK</SccAuxPath>
-    <SccProvider>SAK</SccProvider>
+    <SccProjectName>
+    </SccProjectName>
+    <SccLocalPath>
+    </SccLocalPath>
+    <SccAuxPath>
+    </SccAuxPath>
+    <SccProvider>
+    </SccProvider>
     <IsWebBootstrapper>false</IsWebBootstrapper>
     <FileUpgradeFlags>
     </FileUpgradeFlags>
@@ -161,4 +165,4 @@
   <Target Name="AfterBuild">
   </Target>
   -->
-</Project>
+</Project>
\ No newline at end of file
diff --git a/Engine/Source/Programs/NetworkProfiler/NetworkProfiler/PartialNetworkStream.cs b/Engine/Source/Programs/NetworkProfiler/NetworkProfiler/PartialNetworkStream.cs
--- a/Engine/Source/Programs/NetworkProfiler/NetworkProfiler/PartialNetworkStream.cs
+++ b/Engine/Source/Programs/NetworkProfiler/NetworkProfiler/PartialNetworkStream.cs
@@ -51,11 +51,17 @@ namespace NetworkProfiler
 		/** Number of times SendBunch was called. */
 		public int SendBunchCount = 0;
 		/** Total size of bytes sent. */
-		public int SendBunchSizeBits = 0;		
+		public int SendBunchSizeBits = 0;
+		/** Total size of bunch headers sent. */
+		public int SendBunchHeaderSizeBits = 0;
 		/** Call count per channel type. */
 		public int[] SendBunchCountPerChannel = Enumerable.Repeat(0, (int)EChannelTypes.Max).ToArray();
 		/** Size per channel type */
 		public int[] SendBunchSizeBitsPerChannel = Enumerable.Repeat(0, (int)EChannelTypes.Max).ToArray();
+		/** Size of bunch headers per channel type */
+		public int[] SendBunchHeaderSizeBitsPerChannel = Enumerable.Repeat(0, (int)EChannelTypes.Max).ToArray();
+		/** Size of bunch payloads per channel type */
+		public int[] SendBunchPayloadSizeBitsPerChannel = Enumerable.Repeat(0, (int)EChannelTypes.Max).ToArray();
 
 		// Low level socket
 
@@ -75,6 +81,38 @@ namespace NetworkProfiler
 		/** Total size of acks sent. */
 		public int SendAckSizeBits = 0;
 
+		// Exported GUIDs
+
+		/** Number of GUID export bunches sent. */
+		public int ExportBunchCount = 0;
+		/** Total size of exported GUIDs sent. */
+		public int ExportBunchSizeBits = 0;
+
+		// Must be mapped GUIDs
+
+		/** Number of must be mapped GUIDs. */
+		public int MustBeMappedGuidCount = 0;
+		/** Total size of exported GUIDs sent. */
+		public int MustBeMappedGuidSizeBits = 0;
+
+		// Content blocks
+
+		/** Number of content block headers. */
+		public int ContentBlockHeaderCount = 0;
+		/** Total size of content block headers. */
+		public int ContentBlockHeaderSizeBits = 0;
+		/** Number of content block footers. */
+		public int ContentBlockFooterCount = 0;
+		/** Total size of content block footers. */
+		public int ContentBlockFooterSizeBits = 0;
+
+		// Property handles
+
+		/** Number of property handles. */
+		public int PropertyHandleCount = 0;
+		/** Total size of property handles. */
+		public int PropertyHandleSizeBits = 0;
+
 		// Detailed information.
 
 		/** List of all tokens in this substream. */
@@ -211,8 +249,11 @@ namespace NetworkProfiler
 						var TokenSendBunch = (TokenSendBunch) Token;
 						SendBunchCount++;
 						SendBunchSizeBits += TokenSendBunch.GetNumTotalBits();
+						SendBunchHeaderSizeBits += TokenSendBunch.NumHeaderBits;
 						SendBunchCountPerChannel[TokenSendBunch.ChannelType]++;
 						SendBunchSizeBitsPerChannel[TokenSendBunch.ChannelType] += TokenSendBunch.GetNumTotalBits();
+						SendBunchHeaderSizeBitsPerChannel[TokenSendBunch.ChannelType] += TokenSendBunch.NumHeaderBits;
+						SendBunchPayloadSizeBitsPerChannel[TokenSendBunch.ChannelType] += TokenSendBunch.NumPayloadBits;
 						break;
 					case ETokenTypes.SendRPC:
 						var TokenSendRPC = (TokenSendRPC) Token;
@@ -232,6 +273,14 @@ namespace NetworkProfiler
 								ReplicatedSizeBits += Property.NumBits;
 							}
 						}
+
+						foreach( var PropertyHeader in TokenReplicateActor.PropertyHeaders )
+						{
+							if( PropertyHeader.MatchesFilters( ActorFilter, PropertyFilter, RPCFilter ) )
+							{
+								ReplicatedSizeBits += PropertyHeader.NumBits;
+							}
+						}
 						break;
 					case ETokenTypes.Event:
 						NumEvents++;
@@ -243,6 +292,31 @@ namespace NetworkProfiler
 						SendAckCount++;
 						SendAckSizeBits += TokenSendAck.NumBits;
 						break;
+					case ETokenTypes.ExportBunch:
+						var TokenExportBunch = (TokenExportBunch) Token;
+						ExportBunchCount++;
+						ExportBunchSizeBits += TokenExportBunch.NumBits;
+						break;
+					case ETokenTypes.MustBeMappedGuids:
+						var TokenMustBeMappedGuids = (TokenMustBeMappedGuids) Token;
+						MustBeMappedGuidCount += TokenMustBeMappedGuids.NumGuids;
+						MustBeMappedGuidSizeBits += TokenMustBeMappedGuids.NumBits;
+						break;
+					case ETokenTypes.BeginContentBlock:
+						var TokenBeginContentBlock = (TokenBeginContentBlock) Token;
+						ContentBlockHeaderCount++;
+						ContentBlockHeaderSizeBits += TokenBeginContentBlock.NumBits;
+						break;
+					case ETokenTypes.EndContentBlock:
+						var TokenEndContentBlock = (TokenEndContentBlock) Token;
+						ContentBlockFooterCount++;
+						ContentBlockFooterSizeBits += TokenEndContentBlock.NumBits;
+						break;
+					case ETokenTypes.WritePropertyHandle:
+						var TokenWritePropertyHandle = (TokenWritePropertyHandle) Token;
+						PropertyHandleCount++;
+						PropertyHandleSizeBits += TokenWritePropertyHandle.NumBits;
+						break;
 					default:
 						throw new System.IO.InvalidDataException();
 				}
@@ -405,73 +479,46 @@ namespace NetworkProfiler
 		{
 			return Count.ToString("0.0").PadLeft(8);
 		}
+		
+		/**
+		 * Returns the total number of bits used for bunch format protocol information.
+		 * This includes bunch headers, content block headers and footers, property handles,
+		 * exported GUIDs, and "must be mapped" GUIDs.
+		 */
+		public int GetTotalBunchOverheadSizeBits()
+		{
+			return	SendBunchHeaderSizeBits +
+					ContentBlockHeaderSizeBits +
+					ContentBlockFooterSizeBits +
+					PropertyHandleSizeBits +
+					ExportBunchSizeBits +
+					MustBeMappedGuidSizeBits;
+		}
 
 		/**
 		 * Converts the summary to a human readable array of strings.
 		 */		
 		public string[] ToStringArray()
 		{
-			string FormatString =	"Data Summary^"						+
-									"^"									+
-									"Frame Count            : {0}^"		+
-									"Duration (ms)          : {1}^"		+
-									"Duration (sec)         : {2}^"		+
-									"^"									+
-									"^"									+
-									"Network Summary^"					+
-									"^"									+
-									"Actor Count            : {3}^"		+
-									"Property Count         : {4}^"		+
-									"Replicated Size        : {5}^"		+
-									"RPC Count              : {6}^"		+
-									"RPC Size               : {7}^"		+
-									"Sent ack count         : {8}^"		+
-									"Sent ack size          : {9}^"		+
-									"SendBunch Count        : {10}^"	+
-									"   Control             : {11}^"	+
-									"   Actor               : {12}^"	+
-									"   File                : {13}^"	+
-									"   Voice               : {14}^"	+
-									"SendBunch Size         : {15}^"	+
-									"   Control             : {16}^"	+
-									"   Actor               : {17}^"	+
-									"   File                : {18}^"	+
-									"   Voice               : {19}^"	+
-									"Game Socket Send Count : {20}^"	+
-									"Game Socket Send Size  : {21}^"	+
-									"Misc Socket Send Count : {22}^"	+
-									"Misc Socket Send Size  : {23}^"	+
-									"Outgoing bandwidth     : {24}^"	+
-									"^"									+
-									"^"									+
-									"Network Summary per second^"		+
-									"^"									+
-									"Actor Count            : {25}^"	+
-									"Property Count         : {26}^"	+
-									"Replicated Size        : {27}^"	+
-									"RPC Count              : {28}^"	+
-									"RPC Size               : {29}^"	+
-									"Sent ack count         : {30}^"	+
-									"Sent ack size          : {31}^"	+
-									"SendBunch Count        : {32}^"	+
-									"   Control             : {33}^"	+
-									"   Actor               : {34}^"	+
-									"   File                : {35}^"	+
-									"   Voice               : {36}^"	+
-									"SendBunch Size         : {37}^"	+
-									"   Control             : {38}^"	+
-									"   Actor               : {39}^"	+
-									"   File                : {40}^"	+
-									"   Voice               : {41}^"	+
-									"Game Socket Send Count : {42}^"	+
-									"Game Socket Send Size  : {43}^"	+
-									"Misc Socket Send Count : {44}^"	+
-									"Misc Socket Send Size  : {45}^"	+
-									"Outgoing bandwidth     : {46}^";
-
-			float OneOverDeltaTime = 1 / (EndTime - StartTime);
-
-			string Summary = String.Format( FormatString, 
+			string FormatStringPrefix =	"Data Summary^"					+
+										"^"								+
+										"Frame Count            : {0}^"	+
+										"Duration (ms)          : {1}^"	+
+										"Duration (sec)         : {2}^"	+
+										"^"								+
+										"^"								+
+										"Network Summary^"				+
+										"^"								+
+										"Actor Count            : {3}^"	+
+										"Property Count         : {4}^"	+
+										"Replicated Size        : {5}^"	+
+										"RPC Count              : {6}^"	+
+										"RPC Size               : {7}^"	+
+										"Sent ack count         : {8}^"	+
+										"Sent ack size          : {9}^"	+
+										"Bunch format overhead  : {10}^";
+
+			string Summary = String.Format( FormatStringPrefix, 
 									ConvertToCountString(NumFrames),
 									ConvertToCountString((EndTime - StartTime) * 1000),
 									ConvertToCountString((EndTime - StartTime)),
@@ -482,6 +529,69 @@ namespace NetworkProfiler
 									ConvertToSizeString(RPCSizeBits / 8.0f),
 									ConvertToCountString(SendAckCount),
 									ConvertToSizeString(SendAckSizeBits / 8.0f),
+									ConvertToSizeString(GetTotalBunchOverheadSizeBits() / 8.0f));
+
+			if ( SendBunchHeaderSizeBits > 0 )
+			{
+				Summary += String.Format("   Bunch headers       : {0}^", ConvertToSizeString(SendBunchHeaderSizeBits / 8.0f));
+			}
+
+			if ( ContentBlockHeaderSizeBits > 0 )
+			{
+				Summary += String.Format("   ContentBlock headers: {0}^", ConvertToSizeString(ContentBlockHeaderSizeBits / 8.0f));
+			}
+
+			if ( ContentBlockFooterSizeBits > 0 )
+			{
+				Summary += String.Format("   ContentBlock footers: {0}^", ConvertToSizeString(ContentBlockFooterSizeBits / 8.0f));
+			}
+
+			if ( PropertyHandleSizeBits > 0 )
+			{
+				Summary += String.Format("   Property handles    : {0}^", ConvertToSizeString(PropertyHandleSizeBits / 8.0f));
+			}
+
+			if ( ExportBunchSizeBits > 0 )
+			{
+				Summary += String.Format("   Exported GUIDs      : {0}^", ConvertToSizeString(ExportBunchSizeBits / 8.0f));
+			}
+
+			if ( MustBeMappedGuidSizeBits > 0 )
+			{
+				Summary += String.Format("   MustBeMapped GUIDs  : {0}^", ConvertToSizeString(MustBeMappedGuidSizeBits / 8.0f));
+			}
+
+			string FormatString =	"SendBunch Count        : {0}^"		+
+									"   Control             : {1}^"		+
+									"   Actor               : {2}^"		+
+									"   File                : {3}^"		+
+									"   Voice               : {4}^"		+
+									"SendBunch Size         : {5}^"		+
+									"   Control             : {6}^"		+
+									"   Actor               : {7}^"		+
+									"   File                : {8}^"		+
+									"   Voice               : {9}^"		+
+									"Game Socket Send Count : {10}^"	+
+									"Game Socket Send Size  : {11}^"	+
+									"Misc Socket Send Count : {12}^"	+
+									"Misc Socket Send Size  : {13}^"	+
+									"Outgoing bandwidth     : {14}^"	+
+									"^"									+
+									"^"									+
+									"Network Summary per second^"		+
+									"^"									+
+									"Actor Count            : {15}^"	+
+									"Property Count         : {16}^"	+
+									"Replicated Size        : {17}^"	+
+									"RPC Count              : {18}^"	+
+									"RPC Size               : {19}^"	+										
+									"Sent ack count         : {20}^"	+
+									"Sent ack size          : {21}^"	+
+									"Bunch format overhead  : {22}^";
+			
+			float OneOverDeltaTime = 1 / (EndTime - StartTime);
+
+			Summary += String.Format( FormatString, 
 									ConvertToCountString(SendBunchCount),
 									ConvertToCountString(SendBunchCountPerChannel[(int)EChannelTypes.Control]),
 									ConvertToCountString(SendBunchCountPerChannel[(int)EChannelTypes.Actor]),
@@ -504,6 +614,56 @@ namespace NetworkProfiler
 									ConvertToSizeString(RPCSizeBits / 8.0f * OneOverDeltaTime),
 									ConvertToCountString(SendAckCount * OneOverDeltaTime),
 									ConvertToSizeString(SendAckSizeBits / 8.0f * OneOverDeltaTime),
+									ConvertToSizeString(GetTotalBunchOverheadSizeBits() / 8.0f * OneOverDeltaTime)
+									);
+
+			if ( SendBunchHeaderSizeBits > 0 )
+			{
+				Summary += String.Format("   Bunch headers       : {0}^", ConvertToSizeString(SendBunchHeaderSizeBits / 8.0f * OneOverDeltaTime));
+			}
+
+			if ( ContentBlockHeaderSizeBits > 0 )
+			{
+				Summary += String.Format("   ContentBlock headers: {0}^", ConvertToSizeString(ContentBlockHeaderSizeBits / 8.0f * OneOverDeltaTime));
+			}
+
+			if ( ContentBlockFooterSizeBits > 0 )
+			{
+				Summary += String.Format("   ContentBlock footers: {0}^", ConvertToSizeString(ContentBlockFooterSizeBits / 8.0f * OneOverDeltaTime));
+			}
+
+			if ( PropertyHandleSizeBits > 0 )
+			{
+				Summary += String.Format("   Property handles    : {0}^", ConvertToSizeString(PropertyHandleSizeBits / 8.0f * OneOverDeltaTime));
+			}
+
+			if ( ExportBunchSizeBits > 0 )
+			{
+				Summary += String.Format("   Exported GUIDs      : {0}^", ConvertToSizeString(ExportBunchSizeBits / 8.0f * OneOverDeltaTime));
+			}
+
+			if ( MustBeMappedGuidSizeBits > 0 )
+			{
+				Summary += String.Format("   MustBeMapped GUIDs  : {0}^", ConvertToSizeString(MustBeMappedGuidSizeBits / 8.0f * OneOverDeltaTime));
+			}
+
+			string FormatStringSuffix =	"SendBunch Count        : {0}^"		+
+										"   Control             : {1}^"		+
+										"   Actor               : {2}^"		+
+										"   File                : {3}^"		+
+										"   Voice               : {4}^"		+
+										"SendBunch Size         : {5}^"		+
+										"   Control             : {6}^"		+
+										"   Actor               : {7}^"		+
+										"   File                : {8}^"		+
+										"   Voice               : {9}^"		+
+										"Game Socket Send Count : {10}^"	+
+										"Game Socket Send Size  : {11}^"	+
+										"Misc Socket Send Count : {12}^"	+
+										"Misc Socket Send Size  : {13}^"	+
+										"Outgoing bandwidth     : {14}^";
+
+			Summary += String.Format( FormatStringSuffix,
 									ConvertToCountString(SendBunchCount * OneOverDeltaTime),
 									ConvertToCountString(SendBunchCountPerChannel[(int)EChannelTypes.Control] * OneOverDeltaTime),
 									ConvertToCountString(SendBunchCountPerChannel[(int)EChannelTypes.Actor] * OneOverDeltaTime),
@@ -582,6 +742,21 @@ namespace NetworkProfiler
 		}
 	}
 
+	/**
+	 * UniquePropertyHeader
+	 * Tracks all the unique property headers of the same type
+	*/
+	class UniquePropertyHeader : UniqueItem<TokenWritePropertyHeader>
+	{
+		public long		SizeBits = 0;
+
+		override public void OnItemAdded(TokenWritePropertyHeader Item)
+		{
+			base.OnItemAdded(Item);
+			SizeBits += Item.NumBits;
+		}
+	}
+
 	/**
 	 * UniqueActor
 	 * Tracks all the unique actors of the same type
@@ -603,6 +778,11 @@ namespace NetworkProfiler
 				SizeBits += Property.NumBits;
 				Properties.AddItem(Property, Property.PropertyNameIndex);
 			}
+
+			foreach (var PropertyHeader in Item.PropertyHeaders)
+			{
+				SizeBits += PropertyHeader.NumBits;
+			}
 		}
 	}
 }
diff --git a/Engine/Source/Programs/NetworkProfiler/NetworkProfiler/StreamParser.cs b/Engine/Source/Programs/NetworkProfiler/NetworkProfiler/StreamParser.cs
--- a/Engine/Source/Programs/NetworkProfiler/NetworkProfiler/StreamParser.cs
+++ b/Engine/Source/Programs/NetworkProfiler/NetworkProfiler/StreamParser.cs
@@ -36,13 +36,19 @@ namespace NetworkProfiler
 		 * @param	TokenReplicateActor		Actor token
 		 * @param	LastProperties			Properties to be copied to the actor
 		 */
-		private static void FinishActorProperties(TokenReplicateActor TokenReplicateActor, List<TokenReplicateProperty> LastProperties)
+		private static void FinishActorProperties(TokenReplicateActor TokenReplicateActor, List<TokenReplicateProperty> LastProperties, List<TokenWritePropertyHeader> LastPropertyHeaders)
 		{
 			for (int i = 0; i < LastProperties.Count; i++ )
 			{
 				TokenReplicateActor.Properties.Add(LastProperties[i]);
 			}
 			LastProperties.Clear();
+
+			for (int i = 0; i < LastPropertyHeaders.Count; i++ )
+			{
+				TokenReplicateActor.PropertyHeaders.Add(LastPropertyHeaders[i]);
+			}
+			LastPropertyHeaders.Clear();
 		}
 
 		/**
@@ -88,6 +94,7 @@ namespace NetworkProfiler
 			var CurrentFrameTokens = new List<TokenBase>();
 			TokenReplicateActor LastActorToken = null;
             List<TokenReplicateProperty> LastProperties = new List<TokenReplicateProperty>();
+			List<TokenWritePropertyHeader> LastPropertyHeaders = new List<TokenWritePropertyHeader>();
 
 			TokenFrameMarker LastFrameMarker = null;
 			
@@ -117,6 +124,7 @@ namespace NetworkProfiler
 					CurrentFrameTokens.Clear();
 
 					Debug.Assert(LastProperties.Count == 0);		// We shouldn't have any properties now
+					Debug.Assert(LastPropertyHeaders.Count == 0);	// We shouldn't have any property headers now either
 
 					// Finish up actor summary of last pending actor before switching frames.
 					HandleActorSummary(NetworkStream, LastActorToken);
@@ -132,6 +140,7 @@ namespace NetworkProfiler
 				if( Token.TokenType == ETokenTypes.EndOfStreamMarker )
 				{
 					Debug.Assert(LastProperties.Count == 0);		// We shouldn't have any properties now
+					Debug.Assert(LastPropertyHeaders.Count == 0);	// We shouldn't have any property headers now either
 					bHasReachedEndOfStream = true;
 					// Finish up actor summary of last pending actor at end of stream
 					HandleActorSummary(NetworkStream, LastActorToken);
@@ -143,8 +152,9 @@ namespace NetworkProfiler
 					if( Token.TokenType == ETokenTypes.ReplicateActor )
 					{
 						// Encountered a new actor so we can finish up existing one for summary.
-						FinishActorProperties(Token as TokenReplicateActor, LastProperties );
+						FinishActorProperties(Token as TokenReplicateActor, LastProperties, LastPropertyHeaders );
 						Debug.Assert(LastProperties.Count == 0);		// We shouldn't have any properties now
+						Debug.Assert(LastPropertyHeaders.Count == 0);	// We shouldn't have any property headers now either
 						HandleActorSummary(NetworkStream, LastActorToken);
 						LastActorToken = Token as TokenReplicateActor;
 					}
@@ -162,7 +172,12 @@ namespace NetworkProfiler
 						NetworkStream.UpdateSummary(ref NetworkStream.PropertyNameToSummary, TokenReplicateProperty.PropertyNameIndex, TokenReplicateProperty.NumBits, 0 );
 						//LastActorToken.Properties.Add(TokenReplicateProperty);
                         LastProperties.Add(TokenReplicateProperty);
-					}					
+					}
+					else if( Token.TokenType == ETokenTypes.WritePropertyHeader )
+					{
+						var TokenWritePropertyHeader = Token as TokenWritePropertyHeader;
+                        LastPropertyHeaders.Add(TokenWritePropertyHeader);
+					}
 					else
 					{
 						CurrentFrameTokens.Add(Token);
diff --git a/Engine/Source/Programs/NetworkProfiler/NetworkProfiler/Tokens.cs b/Engine/Source/Programs/NetworkProfiler/NetworkProfiler/Tokens.cs
--- a/Engine/Source/Programs/NetworkProfiler/NetworkProfiler/Tokens.cs
+++ b/Engine/Source/Programs/NetworkProfiler/NetworkProfiler/Tokens.cs
@@ -22,6 +22,12 @@ namespace NetworkProfiler
 		Event,						// Event
 		RawSocketData,				// Raw socket data being sent
 		SendAck,					// Ack being sent
+		WritePropertyHeader,		// Property header being written
+		ExportBunch,				// Exported GUIDs
+		MustBeMappedGuids,			// Must be mapped GUIDs
+		BeginContentBlock,			// Content block headers
+		EndContentBlock,			// Content block footers
+		WritePropertyHandle,		// Property handles
 		MaxAndInvalid,				// Invalid token, also used as the max token index
 	}
 
@@ -95,6 +101,24 @@ namespace NetworkProfiler
 				case ETokenTypes.SendAck:
 					SerializedToken = new TokenSendAck( BinaryStream );
 					break;
+				case ETokenTypes.WritePropertyHeader:
+					SerializedToken = new TokenWritePropertyHeader( BinaryStream );
+					break;
+				case ETokenTypes.ExportBunch:
+					SerializedToken = new TokenExportBunch( BinaryStream );
+					break;
+				case ETokenTypes.MustBeMappedGuids:
+					SerializedToken = new TokenMustBeMappedGuids( BinaryStream );
+					break;
+				case ETokenTypes.BeginContentBlock:
+					SerializedToken = new TokenBeginContentBlock( BinaryStream );
+					break;
+				case ETokenTypes.EndContentBlock:
+					SerializedToken = new TokenEndContentBlock( BinaryStream );
+					break;
+				case ETokenTypes.WritePropertyHandle:
+					SerializedToken = new TokenWritePropertyHandle( BinaryStream );
+					break;
 				default:
 					throw new InvalidDataException();
 			}
@@ -351,6 +375,9 @@ namespace NetworkProfiler
         /** List of property tokens that were serialized for this actor. */
 		public List<TokenReplicateProperty> Properties;
 
+		/** List of property header tokens that were serialized for this actor. */
+		public List<TokenWritePropertyHeader> PropertyHeaders;
+
 		/** Constructor, serializing members from passed in stream. */
 		public TokenReplicateActor(BinaryReader BinaryStream)
 		{
@@ -361,6 +388,7 @@ namespace NetworkProfiler
 			ActorNameIndex = BinaryStream.ReadInt32();
             TimeInMS = BinaryStream.ReadSingle();
 			Properties = new List<TokenReplicateProperty>();
+			PropertyHeaders = new List<TokenWritePropertyHeader>();
 		}
 
 		/**
@@ -380,6 +408,15 @@ namespace NetworkProfiler
 					NumReplicatedBits += Property.NumBits;
 				}
 			}
+
+			foreach( var PropertyHeader in PropertyHeaders )
+			{
+				if( PropertyHeader.MatchesFilters( ActorFilter, PropertyFilter, RPCFilter ) )
+				{
+					NumReplicatedBits += PropertyHeader.NumBits;
+				}
+			}
+
 			return NumReplicatedBits;
 		}
 	 
@@ -400,14 +437,32 @@ namespace NetworkProfiler
 			ReturnList.Add( "   Flags         : " + (bNetDirty ? "bNetDirty " : "") + (bNetInitial ? "bNetInitial" : "") + (bNetOwner ? "bNetOwner" : "") );
 			ReturnList.Add( "   NumBits       : " + NumReplicatedBits );
 			ReturnList.Add( "   NumBytes      : " + NumReplicatedBits / 8.0f );
-			ReturnList.Add( "   PROPERTIES" );
-			foreach( var Property in  Properties )
+			
+			if( Properties.Count > 0 )
 			{
-				if( Property.MatchesFilters( ActorFilter, PropertyFilter, RPCFilter ) )
+				ReturnList.Add( "   PROPERTIES" );
+				foreach( var Property in  Properties )
+				{
+					if( Property.MatchesFilters( ActorFilter, PropertyFilter, RPCFilter ) )
+					{
+						ReturnList.Add( "      Property      : " + NetworkStream.GetName(Property.PropertyNameIndex) );
+						ReturnList.Add( "      NumBits       : " + Property.NumBits );
+						ReturnList.Add( "      NumBytes      : " + Property.NumBits / 8.0f);
+					}
+				}
+			}
+
+			if( PropertyHeaders.Count > 0 )
+			{
+				ReturnList.Add( "   PROPERTY HEADERS" );
+				foreach( var PropertyHeader in  PropertyHeaders )
 				{
-					ReturnList.Add( "      Property      : " + NetworkStream.GetName(Property.PropertyNameIndex) );
-					ReturnList.Add( "      NumBits       : " + Property.NumBits );
-					ReturnList.Add( "      NumBytes      : " + Property.NumBits / 8.0f);
+					if( PropertyHeader.MatchesFilters( ActorFilter, PropertyFilter, RPCFilter ) )
+					{
+						ReturnList.Add( "      Property      : " + NetworkStream.GetName(PropertyHeader.PropertyNameIndex) );
+						ReturnList.Add( "      NumBits       : " + PropertyHeader.NumBits );
+						ReturnList.Add( "      NumBytes      : " + PropertyHeader.NumBits / 8.0f);
+					}
 				}
 			}
 			return ReturnList;
@@ -469,6 +524,176 @@ namespace NetworkProfiler
 		}
 	}			
 
+	/**
+	 * Token for property header replication. Context determines which actor this belongs to.
+	 */
+	class TokenWritePropertyHeader : TokenBase
+	{
+		/** Name table index of property name. */
+		public int		PropertyNameIndex;
+		/** Number of bits serialized/ sent. */
+		public UInt16	NumBits;
+
+		/** Constructor, serializing members from passed in stream. */
+		public TokenWritePropertyHeader(BinaryReader BinaryStream)
+		{
+            PropertyNameIndex = BinaryStream.ReadInt32();
+			NumBits = BinaryStream.ReadUInt16();
+		}
+
+		/**
+		 * Returns whether the token matches/ passes based on the passed in filters.
+		 * 
+		 * @param	ActorFilter		Actor filter to match against
+		 * @param	PropertyFilter	Property filter to match against
+		 * @param	RPCFilter		RPC filter to match against
+		 * 
+		 * @return true if it matches, false otherwise
+		 */
+		public override bool MatchesFilters( string ActorFilter, string PropertyFilter, string RPCFilter )
+		{
+			return PropertyFilter.Length == 0 || NetworkStream.GetName(PropertyNameIndex).ToUpperInvariant().Contains( PropertyFilter.ToUpperInvariant() );
+		}
+	}
+
+	/**
+	 * Token for exported GUID bunches.
+	 */
+	class TokenExportBunch : TokenBase
+	{
+		/** Number of bits serialized/ sent. */
+		public UInt16	NumBits;
+
+		/** Constructor, serializing members from passed in stream. */
+		public TokenExportBunch(BinaryReader BinaryStream)
+		{
+			NumBits = BinaryStream.ReadUInt16();
+		}
+
+		/**
+		 * Converts the token into a multi-string description.
+		 */
+		public override List<string> ToDetailedStringList( string ActorFilter, string PropertyFilter, string RPCFilter )
+		{
+			var ReturnList = new List<string>();
+			ReturnList.Add( "EXPORTED GUIDS" );
+			ReturnList.Add( "   NumBits          : " + NumBits );
+			ReturnList.Add( "   NumBytes         : " + NumBits / 8.0f);
+			return ReturnList;
+		}
+	}
+
+	/**
+	 * Token for must be mapped GUIDs.
+	 */
+	class TokenMustBeMappedGuids : TokenBase
+	{
+		/** Number of GUIDs serialized/sent. */
+		public UInt16	NumGuids;
+
+		/** Number of bits serialized/sent. */
+		public UInt16	NumBits;
+
+		/** Constructor, serializing members from passed in stream. */
+		public TokenMustBeMappedGuids(BinaryReader BinaryStream)
+		{
+			NumGuids = BinaryStream.ReadUInt16();
+			NumBits = BinaryStream.ReadUInt16();
+		}
+
+		/**
+		 * Converts the token into a multi-string description.
+		 */
+		public override List<string> ToDetailedStringList( string ActorFilter, string PropertyFilter, string RPCFilter )
+		{
+			var ReturnList = new List<string>();
+			ReturnList.Add( "MUST BE MAPPED GUIDS" );
+			ReturnList.Add( "   NumGuids         : " + NumGuids );
+			ReturnList.Add( "   NumBits          : " + NumBits );
+			ReturnList.Add( "   NumBytes         : " + NumBits / 8.0f);
+			return ReturnList;
+		}
+	}
+
+	/**
+	 * Token for content block headers.
+	 */
+	class TokenBeginContentBlock : TokenBase
+	{
+		/** Name table index of property name. */
+		public int		ObjectNameIndex;
+		/** Number of bits serialized/ sent. */
+		public UInt16	NumBits;
+
+		/** Constructor, serializing members from passed in stream. */
+		public TokenBeginContentBlock(BinaryReader BinaryStream)
+		{
+            ObjectNameIndex = BinaryStream.ReadInt32();
+			NumBits = BinaryStream.ReadUInt16();
+		}
+
+		/**
+		 * Returns whether the token matches/ passes based on the passed in filters.
+		 * 
+		 * @param	ActorFilter		Actor filter to match against
+		 * @param	PropertyFilter	Property filter to match against
+		 * @param	RPCFilter		RPC filter to match against
+		 * 
+		 * @return true if it matches, false otherwise
+		 */
+		public override bool MatchesFilters( string ActorFilter, string PropertyFilter, string RPCFilter )
+		{
+			return ActorFilter.Length == 0 || NetworkStream.GetName(ObjectNameIndex).ToUpperInvariant().Contains( ActorFilter.ToUpperInvariant() );
+		}
+	}
+
+	/**
+	 * Token for property header replication. Context determines which actor this belongs to.
+	 */
+	class TokenEndContentBlock : TokenBase
+	{
+		/** Name table index of property name. */
+		public int		ObjectNameIndex;
+		/** Number of bits serialized/ sent. */
+		public UInt16	NumBits;
+
+		/** Constructor, serializing members from passed in stream. */
+		public TokenEndContentBlock(BinaryReader BinaryStream)
+		{
+            ObjectNameIndex = BinaryStream.ReadInt32();
+			NumBits = BinaryStream.ReadUInt16();
+		}
+
+		/**
+		 * Returns whether the token matches/ passes based on the passed in filters.
+		 * 
+		 * @param	ActorFilter		Actor filter to match against
+		 * @param	PropertyFilter	Property filter to match against
+		 * @param	RPCFilter		RPC filter to match against
+		 * 
+		 * @return true if it matches, false otherwise
+		 */
+		public override bool MatchesFilters( string ActorFilter, string PropertyFilter, string RPCFilter )
+		{
+			return ActorFilter.Length == 0 || NetworkStream.GetName(ObjectNameIndex).ToUpperInvariant().Contains( ActorFilter.ToUpperInvariant() );
+		}
+	}
+
+	/**
+	 * Token for property hgandle replication. Context determines which actor this belongs to.
+	 */
+	class TokenWritePropertyHandle : TokenBase
+	{
+		/** Number of bits serialized/ sent. */
+		public UInt16	NumBits;
+
+		/** Constructor, serializing members from passed in stream. */
+		public TokenWritePropertyHandle(BinaryReader BinaryStream)
+		{
+			NumBits = BinaryStream.ReadUInt16();
+		}
+	}
+
 	/**
 	 * Token for events.
 	 */
diff --git a/Engine/Source/Runtime/Engine/Private/DataChannel.cpp b/Engine/Source/Runtime/Engine/Private/DataChannel.cpp
--- a/Engine/Source/Runtime/Engine/Private/DataChannel.cpp
+++ b/Engine/Source/Runtime/Engine/Private/DataChannel.cpp
@@ -624,6 +624,8 @@ void UChannel::AppendMustBeMappedGuids( FOutBunch* Bunch )
 			*Bunch << MustBeMappedGuidsInLastBunch[i];
 		}
 
+		NETWORK_PROFILER(GNetworkProfiler.TrackMustBeMappedGuids(NumMustBeMappedGUIDs, Bunch->GetNumBits()));
+
 		// Append the original bunch data at the end
 		Bunch->SerializeBits( TempBunch.GetData(), TempBunch.GetNumBits() );
 
@@ -637,6 +639,15 @@ void UActorChannel::AppendExportBunches( TArray<FOutBunch *>& OutExportBunches )
 {
 	Super::AppendExportBunches( OutExportBunches );
 
+	// Let the profiler know about exported GUID bunches
+	for (const FOutBunch* ExportBunch : QueuedExportBunches )
+	{
+		if (ExportBunch != nullptr)
+		{
+			NETWORK_PROFILER(GNetworkProfiler.TrackExportBunch(ExportBunch->GetNumBits()));
+		}
+	}
+
 	if ( QueuedExportBunches.Num() )
 	{
 		OutExportBunches.Append( QueuedExportBunches );
@@ -2221,6 +2232,8 @@ void UActorChannel::StartBecomingDormant()
 
 void UActorChannel::BeginContentBlock( UObject* Obj, FOutBunch &Bunch )
 {
+	const int NumStartingBits = Bunch.GetNumBits();
+
 #if !(UE_BUILD_SHIPPING || UE_BUILD_TEST)
 	if (CVarDoReplicationContextString->GetInt() > 0)
 	{
@@ -2235,6 +2248,7 @@ void UActorChannel::BeginContentBlock( UObject* Obj, FOutBunch &Bunch )
 
 	if ( IsActor )
 	{
+		NETWORK_PROFILER(GNetworkProfiler.TrackBeginContentBlock(Obj, Bunch.GetNumBits() - NumStartingBits));
 		return;
 	}
 
@@ -2263,12 +2277,16 @@ void UActorChannel::BeginContentBlock( UObject* Obj, FOutBunch &Bunch )
 		Connection->PackageMap->ClearDebugContextString();
 	}
 #endif
+
+	NETWORK_PROFILER(GNetworkProfiler.TrackBeginContentBlock(Obj, Bunch.GetNumBits() - NumStartingBits));
 }
 
 void UActorChannel::BeginContentBlockForSubObjectDelete( FOutBunch & Bunch, FNetworkGUID & GuidToDelete )
 {
 	check( Connection->Driver->IsServer() );
 
+	const int NumStartingBits = Bunch.GetNumBits();
+
 	// Send a 0 bit to signify we are dealing with sub-objects
 	Bunch.WriteBit( 0 );
 
@@ -2285,11 +2303,17 @@ void UActorChannel::BeginContentBlockForSubObjectDelete( FOutBunch & Bunch, FNet
 	FNetworkGUID InvalidNetGUID;
 	InvalidNetGUID.Reset();
 	Bunch << InvalidNetGUID;
+
+	// Since the subobject has been deleted, we don't have a valid object to pass to the profiler.
+	NETWORK_PROFILER(GNetworkProfiler.TrackBeginContentBlock(nullptr, Bunch.GetNumBits() - NumStartingBits));
 }
 
 void UActorChannel::EndContentBlock( UObject *Obj, FOutBunch &Bunch, FClassNetCache* ClassCache )
 {
 	check(Obj);
+
+	const int NumStartingBits = Bunch.GetNumBits();
+
 	if (!ClassCache)
 	{
 		ClassCache = Connection->Driver->NetCache->GetClassNetCache( Obj->GetClass() );
@@ -2297,6 +2321,8 @@ void UActorChannel::EndContentBlock( UObject *Obj, FOutBunch &Bunch, FClassNetCa
 
 	// Write max int to signify done
 	Bunch.WriteIntWrapped(ClassCache->GetMaxIndex(), ClassCache->GetMaxIndex()+1);
+
+	NETWORK_PROFILER(GNetworkProfiler.TrackEndContentBlock(Obj, Bunch.GetNumBits() - NumStartingBits));
 }
 
 UObject* UActorChannel::ReadContentBlockHeader(FInBunch & Bunch, bool& bObjectDeleted)
diff --git a/Engine/Source/Runtime/Engine/Private/DataReplication.cpp b/Engine/Source/Runtime/Engine/Private/DataReplication.cpp
--- a/Engine/Source/Runtime/Engine/Private/DataReplication.cpp
+++ b/Engine/Source/Runtime/Engine/Private/DataReplication.cpp
@@ -843,8 +843,12 @@ void FObjectReplicator::ReplicateCustomDeltaProperties( FOutBunch & Bunch, int32
 		// Write header, and data to send to the actual bunch
 		RepLayout->WritePropertyHeader( Object, ObjectClass, OwningChannel, It, Bunch, Index, LastIndex, bContentBlockWritten );
 
+		const int NumStartingBits = Bunch.GetNumBits();
+
 		// Send property.
 		Bunch.SerializeBits( TempBitWriter.GetData(), TempBitWriter.GetNumBits() );
+
+		NETWORK_PROFILER(GNetworkProfiler.TrackReplicateProperty(It, Bunch.GetNumBits() - NumStartingBits));
 	}
 }
 
@@ -901,6 +905,8 @@ bool FObjectReplicator::ReplicateProperties( FOutBunch & Bunch, FReplicationFlag
 		Bunch.SerializeBits( RemoteFunctions->GetData(), RemoteFunctions->GetNumBits() );
 		RemoteFunctions->Reset();
 		RemoteFuncInfo.Empty();
+
+		NETWORK_PROFILER(GNetworkProfiler.FlushQueuedRPCs(OwningChannelConnection, Object));
 	}
 
 	// See if we wrote something important (anything but the 'end' int below).
diff --git a/Engine/Source/Runtime/Engine/Private/NetConnection.cpp b/Engine/Source/Runtime/Engine/Private/NetConnection.cpp
--- a/Engine/Source/Runtime/Engine/Private/NetConnection.cpp
+++ b/Engine/Source/Runtime/Engine/Private/NetConnection.cpp
@@ -501,8 +501,6 @@ void UNetConnection::FlushNet(bool bIgnoreSimulation)
 
 		NumPaddingBits += SendBuffer.GetNumBits() - NumBitsPrePadding;
 
-		NETWORK_PROFILER(GNetworkProfiler.FlushOutgoingBunches(this));
-
 		// Send now.
 #if DO_ENABLE_NET_TEST
 		// if the connection is closing/being destroyed/etc we need to send immediately regardless of settings
diff --git a/Engine/Source/Runtime/Engine/Private/NetworkDriver.cpp b/Engine/Source/Runtime/Engine/Private/NetworkDriver.cpp
--- a/Engine/Source/Runtime/Engine/Private/NetworkDriver.cpp
+++ b/Engine/Source/Runtime/Engine/Private/NetworkDriver.cpp
@@ -828,12 +828,14 @@ void UNetDriver::InternalProcessRemoteFunction
 	{
 		Ch->BeginContentBlock(TargetObj, Bunch);
 	}
-		
+	
+	const int NumStartingHeaderBits = Bunch.GetNumBits();
+
 	//UE_LOG(LogScript, Log, TEXT("   Call %s"),Function->GetFullName());
 	check( FieldCache->FieldNetIndex <= ClassCache->GetMaxIndex() );
 	Bunch.WriteIntWrapped(FieldCache->FieldNetIndex, ClassCache->GetMaxIndex()+1);
 
-	const int HeaderBits = Bunch.GetNumBits();
+	const int HeaderBits = Bunch.GetNumBits() - NumStartingHeaderBits;
 
 #if !(UE_BUILD_SHIPPING || UE_BUILD_TEST)
 	Bunch.DebugString = FString::Printf(TEXT("%.2f RPC: %s - %s"), Connection->Driver->Time, *Actor->GetName(), *Function->GetName());
@@ -979,7 +981,7 @@ void UNetDriver::InternalProcessRemoteFunction
 				UE_LOG(LogNetTraffic, Log,		TEXT("      Queing unreliable multicast RPC: %s::%s [%.1f bytes]"), *Actor->GetName(), *Function->GetName(), Bunch.GetNumBits() / 8.f );
 			}
 
-			NETWORK_PROFILER(GNetworkProfiler.TrackSendRPC(Actor, Function, HeaderBits, ParameterBits, FooterBits));
+			NETWORK_PROFILER(GNetworkProfiler.TrackQueuedRPC(Connection, TargetObj, Actor, Function, HeaderBits, ParameterBits, FooterBits));
 			Ch->QueueRemoteFunctionBunch(TargetObj, Function, Bunch);
 		}
 		else
diff --git a/Engine/Source/Runtime/Engine/Private/NetworkProfiler.cpp b/Engine/Source/Runtime/Engine/Private/NetworkProfiler.cpp
--- a/Engine/Source/Runtime/Engine/Private/NetworkProfiler.cpp
+++ b/Engine/Source/Runtime/Engine/Private/NetworkProfiler.cpp
@@ -26,7 +26,9 @@ FNetworkProfiler GNetworkProfiler;
 /** Magic value, determining that file is a network profiler file.				*/
 #define NETWORK_PROFILER_MAGIC						0x1DBF348C
 /** Version of memory profiler. Incremented on serialization changes.			*/
-#define NETWORK_PROFILER_VERSION					6
+#define NETWORK_PROFILER_VERSION					7
+
+static const FString UnknownName("UnknownName");
 
 enum ENetworkProfilingPayloadType
 {
@@ -39,7 +41,13 @@ enum ENetworkProfilingPayloadType
 	NPTYPE_EndOfStreamMarker,			// End of stream marker		
 	NPTYPE_Event,						// Event
 	NPTYPE_RawSocketData,				// Raw socket data being sent
-	NPTYPE_SendAck						// Ack being sent
+	NPTYPE_SendAck,						// Ack being sent
+	NPTYPE_WritePropertyHeader,			// Property header being written
+	NPTYPE_ExportBunch,					// Exported GUIDs
+	NPTYPE_MustBeMappedGuids,			// Must be mapped GUIDs
+	NPTYPE_BeginContentBlock,			// Content block headers
+	NPTYPE_EndContentBlock,				// Content block footers
+	NPTYPE_WritePropertyHandle			// Property handles
 };
 
 
@@ -192,6 +200,49 @@ void FNetworkProfiler::TrackSendRPC( const AActor* Actor, const UFunction* Funct
 	}
 }
 
+void FNetworkProfiler::TrackQueuedRPC( UNetConnection* Connection, UObject* TargetObject, const AActor* Actor, const UFunction* Function, uint16 NumHeaderBits, uint16 NumParameterBits, uint16 NumFooterBits  )
+{
+	if( bIsTrackingEnabled )
+	{
+		SCOPE_LOCK_REF(CriticalSection);
+		
+		FQueuedRPCInfo Info;
+		Info.Connection = Connection;
+		Info.TargetObject = TargetObject;
+		Info.ActorNameIndex = GetNameTableIndex( Actor->GetName() );
+		Info.FunctionNameIndex = GetNameTableIndex( Function->GetName() );
+		Info.NumHeaderBits = NumHeaderBits;
+		Info.NumParameterBits = NumParameterBits;
+		Info.NumFooterBits = NumFooterBits;
+
+		QueuedRPCs.Add(Info);
+	}
+}
+
+void FNetworkProfiler::FlushQueuedRPCs( UNetConnection* Connection, UObject* TargetObject )
+{
+	if( bIsTrackingEnabled )
+	{
+		SCOPE_LOCK_REF(CriticalSection);
+
+		for ( int i = QueuedRPCs.Num() - 1; i >= 0; --i )
+		{
+			if (QueuedRPCs[i].Connection == Connection && QueuedRPCs[i].TargetObject == TargetObject)
+			{
+				uint8 Type = NPTYPE_SendRPC;
+				(*FileWriter) << Type;
+				(*FileWriter) << QueuedRPCs[i].ActorNameIndex;
+				(*FileWriter) << QueuedRPCs[i].FunctionNameIndex;
+				(*FileWriter) << QueuedRPCs[i].NumHeaderBits;
+				(*FileWriter) << QueuedRPCs[i].NumParameterBits;
+				(*FileWriter) << QueuedRPCs[i].NumFooterBits;
+
+				QueuedRPCs.RemoveAtSwap(i);
+			}
+		}
+	}
+}
+
 /**
  * Low level FSocket::Send information.
  *
@@ -387,6 +438,19 @@ void FNetworkProfiler::TrackReplicateProperty( const UProperty* Property, uint16
 	}
 }
 
+void FNetworkProfiler::TrackWritePropertyHeader( const UProperty* Property, uint16 NumBits )
+{
+	if( bIsTrackingEnabled )
+	{
+		SCOPE_LOCK_REF(CriticalSection);
+		uint8 Type = NPTYPE_WritePropertyHeader;
+		(*FileWriter) << Type;
+		int32 NameTableIndex = GetNameTableIndex( Property->GetName() );
+		(*FileWriter) << NameTableIndex;
+		(*FileWriter) << NumBits;
+	}
+}
+
 /**
  * Track event occurring, like e.g. client join/ leave
  *
@@ -514,16 +578,77 @@ void FNetworkProfiler::TrackSessionChange( bool bShouldContinueTracking, const F
 #endif	//#if ALLOW_DEBUG_FILES
 }
 
- void FNetworkProfiler::TrackSendAck( uint16 NumBits )
- {
-	 if ( bIsTrackingEnabled )
-	 {
-		 SCOPE_LOCK_REF(CriticalSection);
-		 uint8 Type = NPTYPE_SendAck;
-		 (*FileWriter) << Type;
-		 (*FileWriter) << NumBits;
+void FNetworkProfiler::TrackSendAck( uint16 NumBits )
+{
+	if ( bIsTrackingEnabled )
+	{
+		SCOPE_LOCK_REF(CriticalSection);
+		uint8 Type = NPTYPE_SendAck;
+		(*FileWriter) << Type;
+		(*FileWriter) << NumBits;
+	}
+}
+
+void FNetworkProfiler::TrackMustBeMappedGuids( uint16 NumGuids, uint16 NumBits )
+{
+	if ( bIsTrackingEnabled )
+	{
+		SCOPE_LOCK_REF(CriticalSection);
+		uint8 Type = NPTYPE_MustBeMappedGuids;
+		(*FileWriter) << Type;
+		(*FileWriter) << NumGuids;
+		(*FileWriter) << NumBits;
+	}
+}
+
+void FNetworkProfiler::TrackExportBunch( uint16 NumBits )
+{
+	if ( bIsTrackingEnabled )
+	{
+		SCOPE_LOCK_REF(CriticalSection);
+		uint8 Type = NPTYPE_ExportBunch;
+		(*FileWriter) << Type;
+		(*FileWriter) << NumBits;
+	}
+}
+
+ 
+void FNetworkProfiler::TrackBeginContentBlock(UObject* Object, uint16 NumBits)
+{
+	if ( bIsTrackingEnabled )
+	{
+		SCOPE_LOCK_REF(CriticalSection);
+		uint8 Type = NPTYPE_BeginContentBlock;
+		(*FileWriter) << Type;
+		int32 NameTableIndex = GetNameTableIndex( Object != nullptr ? Object->GetName() : UnknownName );
+		(*FileWriter) << NameTableIndex;
+		(*FileWriter) << NumBits;
 	 }
- }
+}
+
+void FNetworkProfiler::TrackEndContentBlock(UObject* Object, uint16 NumBits)
+{
+	if ( bIsTrackingEnabled )
+	{
+		SCOPE_LOCK_REF(CriticalSection);
+		uint8 Type = NPTYPE_EndContentBlock;
+		(*FileWriter) << Type;
+		int32 NameTableIndex = GetNameTableIndex( Object != nullptr ? Object->GetName() : UnknownName );
+		(*FileWriter) << NameTableIndex;
+		(*FileWriter) << NumBits;
+	}
+}
+
+void FNetworkProfiler::TrackWritePropertyHandle(uint16 NumBits)
+{
+	if ( bIsTrackingEnabled )
+	{
+		SCOPE_LOCK_REF(CriticalSection);
+		uint8 Type = NPTYPE_WritePropertyHandle;
+		(*FileWriter) << Type;
+		(*FileWriter) << NumBits;
+	}
+}
 
 /**
  * Processes any network profiler specific exec commands
diff --git a/Engine/Source/Runtime/Engine/Private/PackageMapClient.cpp b/Engine/Source/Runtime/Engine/Private/PackageMapClient.cpp
--- a/Engine/Source/Runtime/Engine/Private/PackageMapClient.cpp
+++ b/Engine/Source/Runtime/Engine/Private/PackageMapClient.cpp
@@ -878,6 +878,15 @@ bool UPackageMapClient::AppendExportBunches(TArray<FOutBunch *>& OutgoingBunches
 		ExportNetGUIDHeader();
 	}
 
+	// Let the profiler know about exported GUID bunches
+	for (const FOutBunch* ExportBunch : ExportBunches )
+	{
+		if (ExportBunch != nullptr)
+		{
+			NETWORK_PROFILER(GNetworkProfiler.TrackExportBunch(ExportBunch->GetNumBits()));
+		}
+	}
+
 	// Append the bunches we've made to the passed in list reference
 	if (ExportBunches.Num() > 0)
 	{
diff --git a/Engine/Source/Runtime/Engine/Private/RepLayout.cpp b/Engine/Source/Runtime/Engine/Private/RepLayout.cpp
--- a/Engine/Source/Runtime/Engine/Private/RepLayout.cpp
+++ b/Engine/Source/Runtime/Engine/Private/RepLayout.cpp
@@ -927,6 +927,8 @@ bool FRepLayout::ReadyForDormancy( FRepState * RepState ) const
 
 static FORCEINLINE void WritePropertyHandle( FNetBitWriter & Writer, uint16 Handle, bool bDoChecksum )
 {
+	const int NumStartingBits = Writer.GetNumBits();
+
 	uint32 LocalHandle = Handle;
 	Writer.SerializeIntPacked( LocalHandle );
 
@@ -936,6 +938,8 @@ static FORCEINLINE void WritePropertyHandle( FNetBitWriter & Writer, uint16 Hand
 		SerializeGenericChecksum( Writer );
 	}
 #endif
+
+	NETWORK_PROFILER(GNetworkProfiler.TrackWritePropertyHandle( Writer.GetNumBits() - NumStartingBits ));
 }
 
 static bool ShouldSendProperty( FRepWriterState & WriterState, const uint16 Handle )
@@ -1110,6 +1114,8 @@ void FRepLayout::WritePropertyHeader(
 		bContentBlockWritten = true;
 	}
 
+	const int NumStartingBits = Bunch.GetNumBits();
+
 	// Get the network friend property index to replicate
 	// Swap Role/RemoteRole while we're at it
 	FFieldNetCache * FieldCache
@@ -1135,6 +1141,8 @@ void FRepLayout::WritePropertyHeader(
 		Bunch.SerializeIntPacked(idx);
 		LastArrayIndex = ArrayIndex;
 	}
+
+	NETWORK_PROFILER(GNetworkProfiler.TrackWritePropertyHeader(Property, Bunch.GetNumBits() - NumStartingBits));
 }
 
 void FRepLayout::SendProperties( 
diff --git a/Engine/Source/Runtime/Engine/Public/Net/NetworkProfiler.h b/Engine/Source/Runtime/Engine/Public/Net/NetworkProfiler.h
--- a/Engine/Source/Runtime/Engine/Public/Net/NetworkProfiler.h
+++ b/Engine/Source/Runtime/Engine/Public/Net/NetworkProfiler.h
@@ -69,6 +69,29 @@ private:
 	/** Stack outgoing bunches per connection, the top bunch for a connection may be popped if it gets merged with a new bunch.		*/
 	TMap<UNetConnection*, TArray<FSendBunchInfo>>	OutgoingBunches;
 
+	/** Data required to write queued RPCs to the profiler stream */
+	struct FQueuedRPCInfo
+	{
+		UNetConnection* Connection;
+		UObject* TargetObject;
+		int32 ActorNameIndex;
+		int32 FunctionNameIndex;
+		uint16 NumHeaderBits;
+		uint16 NumParameterBits;
+		uint16 NumFooterBits;
+
+		FQueuedRPCInfo()
+			: Connection(nullptr)
+			, TargetObject(nullptr)
+			, ActorNameIndex(0)
+			, FunctionNameIndex(0)
+			, NumHeaderBits(0)
+			, NumParameterBits(0)
+			, NumFooterBits(0) {}
+	};
+	
+	TArray<FQueuedRPCInfo> QueuedRPCs;
+
 	/**
 	 * Returns index of passed in name into name array. If not found, adds it.
 	 *
@@ -106,6 +129,27 @@ public:
 	 */
 	void TrackSendRPC( const AActor* Actor, const UFunction* Function, uint16 NumHeaderBits, uint16 NumParameterBits, uint16 NumFooterBits );
 	
+	/**
+	 * Tracks queued RPCs (unreliable multicast) being sent.
+	 * 
+	 * @param	Connection			The connection on which this RPC is queued
+	 * @param	TargetObject		The target object of the RPC
+	 * @param	Actor				Actor RPC is being called on
+	 * @param	Function			Function being called
+	 * @param	NumHeaderBits		Number of bits serialized into the header for this RPC
+	 * @param	NumParameterBits	Number of bits serialized into parameters of this RPC
+	 * @param	NumFooterBits		Number of bits serialized into the footer of this RPC (EndContentBlock)
+	 */
+	void TrackQueuedRPC( UNetConnection* Connection, UObject* TargetObject, const AActor* Actor, const UFunction* Function, uint16 NumHeaderBits, uint16 NumParameterBits, uint16 NumFooterBits );
+
+	/**
+	 * Writes all queued RPCs for the connection to the profiler stream
+	 *
+	 * @param Connection The connection for which RPCs are being flushed
+	 * @param TargetObject The target object of the RPC
+	 */
+	void FlushQueuedRPCs( UNetConnection* Connection, UObject* TargetObject );
+
 	/**
 	 * Low level FSocket::Send information.
 	 *
@@ -187,7 +231,7 @@ public:
 	 *
 	 * @param Connection the connection which is about to send any pending bunches over the network
 	 */
-	void FlushOutgoingBunches( UNetConnection* Connection );
+	ENGINE_API void FlushOutgoingBunches( UNetConnection* Connection );
 
 	/**
 	 * Track actor being replicated.
@@ -204,6 +248,14 @@ public:
 	 */
 	void TrackReplicateProperty( const UProperty* Property, uint16 NumBits );
 
+	/**
+	 * Track property header being written.
+	 *
+	 * @param	Property	Property being replicated
+	 * @param	NumBits		Number of bits used in the header
+	 */
+	void TrackWritePropertyHeader( const UProperty* Property, uint16 NumBits );
+
 	/**
 	 * Track event occuring, like e.g. client join/ leave
 	 *
@@ -229,6 +281,43 @@ public:
 	 */
 	void TrackSendAck( uint16 NumBits );
 
+	/**
+	 * Track NetGUID export bunches.
+	 *
+	 * @param NumBits Number of bits in the GUIDs
+	 */
+	void TrackExportBunch( uint16 NumBits );
+
+	/**
+	 * Track "must be mapped" GUIDs
+	 *
+	 * @param NumGuids Number of GUIDs added to the bunch
+	 * @param NumBits Number of bits added to the bunch for the GUIDs
+	 */
+	void TrackMustBeMappedGuids( uint16 NumGuids, uint16 NumBits );
+
+	/**
+	 * Track actor content block headers
+	 *
+	 * @param Object the object being replicated (might be a subobject of the actor)
+	 * @param NumBits the number of bits in the content block header
+	 */
+	void TrackBeginContentBlock( UObject* Object, uint16 NumBits );
+
+	/**
+	 * Track actor content block headers
+	 *
+	 * @param Object the object being replicated (might be a subobject of the actor)
+	 * @param NumBits the number of bits in the content block footer
+	 */
+	void TrackEndContentBlock( UObject* Object, uint16 NumBits );
+
+	/** Track property handles
+	 *
+	 * @param NumBits Number of bits in the property handle
+	 */
+	void TrackWritePropertyHandle( uint16 NumBits );
+
 	/**
 	 * Processes any network profiler specific exec commands
 	 *
diff --git a/Engine/Source/Runtime/Online/OnlineSubsystemUtils/Private/IpConnection.cpp b/Engine/Source/Runtime/Online/OnlineSubsystemUtils/Private/IpConnection.cpp
--- a/Engine/Source/Runtime/Online/OnlineSubsystemUtils/Private/IpConnection.cpp
+++ b/Engine/Source/Runtime/Online/OnlineSubsystemUtils/Private/IpConnection.cpp
@@ -133,6 +133,7 @@ void UIpConnection::LowLevelSend( void* Data, int32 Count )
 	CLOCK_CYCLES(Driver->SendCycles);
 	Socket->SendTo((uint8*)Data, Count, BytesSent, *RemoteAddr);
 	UNCLOCK_CYCLES(Driver->SendCycles);
+	NETWORK_PROFILER(GNetworkProfiler.FlushOutgoingBunches(this));
 	NETWORK_PROFILER(GNetworkProfiler.TrackSocketSendTo(Socket->GetDescription(),Data,BytesSent,NumPacketIdBits,NumBunchBits,NumAckBits,NumPaddingBits,*RemoteAddr));
 }
 
