Update expected version for the network profiler (for compatibility with CL 2445520)

diff --git a/Engine/Source/Programs/NetworkProfiler/NetworkProfiler/StreamHeader.cs b/Engine/Source/Programs/NetworkProfiler/NetworkProfiler/StreamHeader.cs
--- a/Engine/Source/Programs/NetworkProfiler/NetworkProfiler/StreamHeader.cs
+++ b/Engine/Source/Programs/NetworkProfiler/NetworkProfiler/StreamHeader.cs
@@ -17,7 +17,7 @@ namespace NetworkProfiler
 		public const UInt32 ExpectedMagic = 0x1DBF348C;
 
 		/** We expect this version, or we can't proceed.			*/
-		public const UInt32 ExpectedVersion = 6;
+		public const UInt32 ExpectedVersion = 7;
 
 		/** Magic to ensure we're opening the right file.			*/
 		public UInt32 Magic;
