PR #5793: Netprofiler: Added "average time" to actor list view (Contributed by [REDACTED_PERSON])

diff --git a/Engine/Source/Programs/NetworkProfiler/NetworkProfiler/MainWindow.cs b/Engine/Source/Programs/NetworkProfiler/NetworkProfiler/MainWindow.cs
--- a/Engine/Source/Programs/NetworkProfiler/NetworkProfiler/MainWindow.cs
+++ b/Engine/Source/Programs/NetworkProfiler/NetworkProfiler/MainWindow.cs
@@ -94,9 +94,9 @@ namespace NetworkProfiler
 			SetDefaultLineView();
 
 			// Force the columns to the way we want them
-			SetupColumns( ActorListView,	new String[]	{ "Total Size (KBytes)", "Count", "Average Size (Bytes)", "Average Size (Bits)", "Time (ms)", "Actor Class" });
-			SetupColumns( PropertyListView, new String[]	{ "Total Size (KBytes)", "Count", "Average Size (Bytes)", "Average Size (Bits)", "Time (ms)", "Property" });
-			SetupColumns( RPCListView,		new String[]	{ "Total Size (KBytes)", "Count", "Average Size (Bytes)", "Average Size (Bits)", "Time (ms)", "RPC" });
+			SetupColumns( ActorListView,	new String[]	{ "Total Size (KBytes)", "Count", "Average Size (Bytes)", "Average Size (Bits)", "Time (ms)", "Average Time (ms)", "Actor Class" });
+			SetupColumns( PropertyListView, new String[]	{ "Total Size (KBytes)", "Count", "Average Size (Bytes)", "Average Size (Bits)", "Time (ms)", "Average Time (ms)", "Property" });
+			SetupColumns( RPCListView,		new String[]	{ "Total Size (KBytes)", "Count", "Average Size (Bytes)", "Average Size (Bits)", "Time (ms)", "Average Time (ms)", "RPC" });
 
 			SetupColumns( ActorPerfPropsListView, new String[] { "Actor", "MS", "KB/s", "Bytes", "Count", "Update HZ", "Rep HZ", "Waste" } );
 			SetupColumns( ActorPerfPropsDetailsListView, new String[] { "Property", "Bytes", "Count" } );
diff --git a/Engine/Source/Programs/NetworkProfiler/NetworkProfiler/StreamParser.cs b/Engine/Source/Programs/NetworkProfiler/NetworkProfiler/StreamParser.cs
--- a/Engine/Source/Programs/NetworkProfiler/NetworkProfiler/StreamParser.cs
+++ b/Engine/Source/Programs/NetworkProfiler/NetworkProfiler/StreamParser.cs
@@ -298,7 +298,7 @@ namespace NetworkProfiler
 			ListView.Items.Clear();
 
 			// Columns are total size KByte, count, avg size in bytes, avg size in bits and associated name.
-			var Columns = new string[6];
+			var Columns = new string[7];
 			foreach( var SummaryEntry in Summaries )
 			{
 				Columns[0] = ((float)SummaryEntry.Value.SizeBits / 8 / 1024).ToString("0.0");
@@ -306,7 +306,8 @@ namespace NetworkProfiler
                 Columns[2] = ((float)SummaryEntry.Value.SizeBits / 8 / SummaryEntry.Value.Count).ToString("0.0");
 				Columns[3] = ((float)SummaryEntry.Value.SizeBits / SummaryEntry.Value.Count).ToString("0.0");
                 Columns[4] = SummaryEntry.Value.TimeInMS.ToString("0.00");
-                Columns[5] = NetworkStream.GetName(SummaryEntry.Key);
+				Columns[5] = (SummaryEntry.Value.TimeInMS / SummaryEntry.Value.Count).ToString("0.0000");
+				Columns[6] = NetworkStream.GetName(SummaryEntry.Key);
                 ListView.Items.Add(new ListViewItem(Columns));
 			}
 
