Fix issue with opening large profiles in the network profiler - we were hitting the array size limit.

diff --git a/Engine/Source/Programs/NetworkProfiler/NetworkProfiler/app.config b/Engine/Source/Programs/NetworkProfiler/NetworkProfiler/app.config
--- a/Engine/Source/Programs/NetworkProfiler/NetworkProfiler/app.config
+++ b/Engine/Source/Programs/NetworkProfiler/NetworkProfiler/app.config
@@ -5,6 +5,7 @@
   </connectionStrings>
   <runtime>
     <generatePublisherEvidence enabled="false"/>
+    <gcAllowVeryLargeObjects enabled="true" />
   </runtime>
   <startup>
     <supportedRuntime version="v4.0" sku=".NETFramework,Version=v4.0,Profile=Client"/>
