Updated Network Profiler to .Net 4.6.2

diff --git a/Engine/Source/Programs/NetworkProfiler/NetworkProfiler/NetworkProfiler.csproj b/Engine/Source/Programs/NetworkProfiler/NetworkProfiler/NetworkProfiler.csproj
--- a/Engine/Source/Programs/NetworkProfiler/NetworkProfiler/NetworkProfiler.csproj
+++ b/Engine/Source/Programs/NetworkProfiler/NetworkProfiler/NetworkProfiler.csproj
@@ -1,5 +1,5 @@
 ﻿<?xml version="1.0" encoding="utf-8"?>
-<Project ToolsVersion="4.0" DefaultTargets="Build" xmlns="http://schemas.microsoft.com/developer/msbuild/2003">
+<Project ToolsVersion="12.0" DefaultTargets="Build" xmlns="http://schemas.microsoft.com/developer/msbuild/2003">
   <PropertyGroup>
     <Configuration Condition=" '$(Configuration)' == '' ">Debug</Configuration>
     <Platform Condition=" '$(Platform)' == '' ">AnyCPU</Platform>
@@ -10,7 +10,7 @@
     <AppDesignerFolder>Properties</AppDesignerFolder>
     <RootNamespace>NetworkProfiler</RootNamespace>
     <AssemblyName>NetworkProfiler</AssemblyName>
-    <TargetFrameworkVersion>v4.5</TargetFrameworkVersion>
+    <TargetFrameworkVersion>v4.6.2</TargetFrameworkVersion>
     <FileAlignment>512</FileAlignment>
     <SccProjectName>
     </SccProjectName>
diff --git a/Engine/Source/Programs/NetworkProfiler/NetworkProfiler/app.config b/Engine/Source/Programs/NetworkProfiler/NetworkProfiler/app.config
--- a/Engine/Source/Programs/NetworkProfiler/NetworkProfiler/app.config
+++ b/Engine/Source/Programs/NetworkProfiler/NetworkProfiler/app.config
@@ -5,10 +5,9 @@
   </connectionStrings>
   <runtime>
     <generatePublisherEvidence enabled="false"/>
-    <gcAllowVeryLargeObjects enabled="true" />
+    <gcAllowVeryLargeObjects enabled="true"/>
   </runtime>
   <startup>
-    <supportedRuntime version="v4.0" sku=".NETFramework,Version=v4.0,Profile=Client"/>
+    <supportedRuntime version="v4.0" sku=".NETFramework,Version=v4.6.2"/>
   </startup>
 </configuration>
-
