Network Profiler:
Update to .NET 4.6.2
Allow files to be dropped onto the window to open them
Check cmdline for a file to open so it can be used by explorer to open nprof files

diff --git a/Engine/Source/Programs/NetworkProfiler/NetworkProfiler/MainWindow.cs b/Engine/Source/Programs/NetworkProfiler/NetworkProfiler/MainWindow.cs
--- a/Engine/Source/Programs/NetworkProfiler/NetworkProfiler/MainWindow.cs
+++ b/Engine/Source/Programs/NetworkProfiler/NetworkProfiler/MainWindow.cs
@@ -85,12 +85,16 @@ namespace NetworkProfiler
 		/** If non 0, we will early out of loading in this many minutes worth of profile time */
 		int MaxProfileMinutes		= 0;
 
-		public MainWindow()
+		public MainWindow( string[] Args )
 		{
 			InitializeComponent();
 
 			this.Closing += OnFormClosing;
 
+			this.AllowDrop = true;
+			this.DragEnter += OnDragEnter;
+			this.DragDrop += OnDragDrop;
+
 			// Set default state.
 			SetDefaultLineView();
 
@@ -108,8 +112,29 @@ namespace NetworkProfiler
 			ActorPerfPropsDetailsListView.Columns[0].Width = 170;
 			ActorPerfPropsDetailsListView.Columns[1].Width = 50;
 			ActorPerfPropsDetailsListView.Columns[2].Width = 50;
+
+			if ( Args.Length > 0 )
+			{
+				ChangeNetworkStream( Args[0] );
+			}
 		}
-		
+		private void OnDragEnter( object sender, DragEventArgs e )
+		{
+			if ( e.Data.GetDataPresent( DataFormats.FileDrop ) )
+			{
+				e.Effect = DragDropEffects.Copy;
+			}
+		}
+
+		private void OnDragDrop( object sender, DragEventArgs e )
+		{
+			string[] Files = ( string[] )e.Data.GetData( DataFormats.FileDrop );
+			if( Files.Length > 0 )
+			{
+				ChangeNetworkStream( Files[0] );
+			}
+		}
+
 		private void OnFormClosing( object sender, CancelEventArgs e )
 		{
 			CancelLoadThread();
diff --git a/Engine/Source/Programs/NetworkProfiler/NetworkProfiler/NetworkProfiler.csproj b/Engine/Source/Programs/NetworkProfiler/NetworkProfiler/NetworkProfiler.csproj
--- a/Engine/Source/Programs/NetworkProfiler/NetworkProfiler/NetworkProfiler.csproj
+++ b/Engine/Source/Programs/NetworkProfiler/NetworkProfiler/NetworkProfiler.csproj
@@ -1,4 +1,4 @@
-﻿<?xml version="1.0" encoding="utf-8"?>
+<?xml version="1.0" encoding="utf-8"?>
 <Project ToolsVersion="12.0" DefaultTargets="Build" xmlns="http://schemas.microsoft.com/developer/msbuild/2003">
   <PropertyGroup>
     <Configuration Condition=" '$(Configuration)' == '' ">Debug</Configuration>
diff --git a/Engine/Source/Programs/NetworkProfiler/NetworkProfiler/Program.cs b/Engine/Source/Programs/NetworkProfiler/NetworkProfiler/Program.cs
--- a/Engine/Source/Programs/NetworkProfiler/NetworkProfiler/Program.cs
+++ b/Engine/Source/Programs/NetworkProfiler/NetworkProfiler/Program.cs
@@ -20,7 +20,7 @@ namespace NetworkProfiler
 			Application.EnableVisualStyles();
 			Application.SetCompatibleTextRenderingDefault(false);
 			Application.CurrentCulture = System.Globalization.CultureInfo.InvariantCulture;
-			Application.Run(new MainWindow());
+			Application.Run(new MainWindow(Args));
 		}
 	}
 }
