Network Profiler: Add time % column to actor list view

diff --git a/Engine/Source/Programs/NetworkProfiler/NetworkProfiler/MainWindow.cs b/Engine/Source/Programs/NetworkProfiler/NetworkProfiler/MainWindow.cs
--- a/Engine/Source/Programs/NetworkProfiler/NetworkProfiler/MainWindow.cs
+++ b/Engine/Source/Programs/NetworkProfiler/NetworkProfiler/MainWindow.cs
@@ -103,7 +103,7 @@ namespace NetworkProfiler
 			SetupColumns( PropertyListView, new String[]	{ "Total Size (KBytes)", "Count", "Average Size (Bytes)", "Average Size (Bits)", "Time (ms)", "Average Time (ms)", "Property" });
 			SetupColumns( RPCListView,		new String[]	{ "Total Size (KBytes)", "Count", "Average Size (Bytes)", "Average Size (Bits)", "Time (ms)", "Average Time (ms)", "RPC" });
 
-			SetupColumns( ActorPerfPropsListView, new String[] { "Actor", "MS", "KB/s", "Bytes", "Count", "Update HZ", "Rep HZ", "Waste" } );
+			SetupColumns( ActorPerfPropsListView, new String[] { "Actor", "MS", "MS %", "KB/s", "Bytes", "Count", "Update HZ", "Rep HZ", "Waste" } );
 			SetupColumns( ActorPerfPropsDetailsListView, new String[] { "Property", "Bytes", "Count" } );
 
 			SetupColumns( ObjectReplicationListView, new String[] { "Object Class", "# Comparisons", "# Replications", "Comparison Time (ms)", "Avg. Time Per Compare (ms)" } );
diff --git a/Engine/Source/Programs/NetworkProfiler/NetworkProfiler/PartialNetworkStream.cs b/Engine/Source/Programs/NetworkProfiler/NetworkProfiler/PartialNetworkStream.cs
--- a/Engine/Source/Programs/NetworkProfiler/NetworkProfiler/PartialNetworkStream.cs
+++ b/Engine/Source/Programs/NetworkProfiler/NetworkProfiler/PartialNetworkStream.cs
@@ -500,20 +500,28 @@ namespace NetworkProfiler
 
 			float OneOverDeltaTime = 1 / ( EndTime - StartTime );
 
-			var Columns = new string[8];
+			float TotalTimeInMS = 0;
+			foreach (var UniqueActor in ActorDetailList)
+			{
+				TotalTimeInMS += UniqueActor.Value.TimeInMS;
+			}
+			float OneOverTotalTimeMS = 1 / TotalTimeInMS;
+
+			var Columns = new string[9];
 			foreach (var UniqueActor in ActorDetailList)
 			{
 				long NumActorBytes = ( UniqueActor.Value.SizeBits + 7 ) / 8;
 
 				Columns[0] = NetworkStream.GetName( UniqueActor.Key );
 				Columns[1] = UniqueActor.Value.TimeInMS.ToString( "0.00" );
-				Columns[2] = ( NumActorBytes * OneOverDeltaTime / 1024.0f ).ToString( "0.00" );
-				Columns[3] = NumActorBytes.ToString();
-				Columns[4] = UniqueActor.Value.Count.ToString();
-                Columns[5] = ( UniqueActor.Value.Count * OneOverDeltaTime ).ToString( "0.00" );
-                Columns[6] = ( UniqueActor.Value.ReplicatedCount * OneOverDeltaTime ).ToString( "0.00" );
-
-				Columns[7] = UniqueActor.Value.Count > 0 ? ( 100.0f - ( ( float )UniqueActor.Value.ReplicatedCount / ( float )UniqueActor.Value.Count * 100.0f ) ).ToString( "0.00" ) : "0";
+				Columns[2] = ( UniqueActor.Value.TimeInMS * OneOverTotalTimeMS * 100 ).ToString("0.00");
+				Columns[3] = ( NumActorBytes * OneOverDeltaTime / 1024.0f ).ToString( "0.00" );
+				Columns[4] = NumActorBytes.ToString();
+				Columns[5] = UniqueActor.Value.Count.ToString();
+                Columns[6] = ( UniqueActor.Value.Count * OneOverDeltaTime ).ToString( "0.00" );
+                Columns[7] = ( UniqueActor.Value.ReplicatedCount * OneOverDeltaTime ).ToString( "0.00" );
+
+				Columns[8] = UniqueActor.Value.Count > 0 ? ( 100.0f - ( ( float )UniqueActor.Value.ReplicatedCount / ( float )UniqueActor.Value.Count * 100.0f ) ).ToString( "0.00" ) : "0";
 
                 ListView.Items.Add( new ListViewItem( Columns ) );
 			}
